Audio Analyzer Infrastructure

diff --git a/Engine/Plugins/Runtime/AudioSynesthesia/AudioSynesthesia.uplugin b/Engine/Plugins/Runtime/AudioSynesthesia/AudioSynesthesia.uplugin
new file mode 100644
--- /dev/null
+++ b/Engine/Plugins/Runtime/AudioSynesthesia/AudioSynesthesia.uplugin
@@ -0,0 +1,30 @@
+{
+	"FileVersion" : 3,
+	"Version" : 1,
+	"VersionName" : "1.0",
+	"FriendlyName" : "Audio Synesthesia",
+	"Description" : "A variety of offline analyzers for integrating exposing extracted audio metadata to blueprints.",
+	"Category" : "Audio",
+	"CreatedBy" : "Epic Games, Inc.",
+	"CreatedByURL" : "http://epicgames.com",
+	"DocsURL" : "",
+	"MarketplaceURL" : "",
+	"SupportURL" : "",
+	"EnabledByDefault" : false,
+	"CanContainContent" : true,
+	"IsBetaVersion" : true,
+	"Installed" : false,
+	"Modules" :
+	[
+		{
+			"Name" : "AudioSynesthesia",
+			"Type" : "Runtime",
+			"LoadingPhase" : "Default"
+		},
+		{
+			"Name" : "AudioSynesthesiaEditor",
+			"Type" : "Editor",
+			"LoadingPhase" : "PostDefault"
+		}
+	]
+}
diff --git a/Engine/Plugins/Runtime/AudioSynesthesia/Source/AudioSynesthesia/AudioSynesthesia.Build.cs b/Engine/Plugins/Runtime/AudioSynesthesia/Source/AudioSynesthesia/AudioSynesthesia.Build.cs
new file mode 100644
--- /dev/null
+++ b/Engine/Plugins/Runtime/AudioSynesthesia/Source/AudioSynesthesia/AudioSynesthesia.Build.cs
@@ -0,0 +1,26 @@
+// Copyright 1998-2019 Epic Games, Inc. All Rights Reserved.
+// Copyright 1998-2018 Epic Games, Inc. All Rights Reserved.
+using UnrealBuildTool;
+using System.IO;
+
+namespace UnrealBuildTool.Rules
+{
+	public class AudioSynesthesia : ModuleRules
+	{
+        public AudioSynesthesia(ReadOnlyTargetRules Target) : base(Target)
+		{
+
+            //OptimizeCode = CodeOptimization.Never;
+			
+            PublicDependencyModuleNames.AddRange(
+				new string[] {
+                    "Core",
+					"CoreUObject",
+					"Engine",
+					"AudioMixer",
+                    "AudioAnalyzer"
+                }
+            );
+		}
+	}
+}
diff --git a/Engine/Plugins/Runtime/AudioSynesthesia/Source/AudioSynesthesia/Classes/AudioSynesthesiaNRT.h b/Engine/Plugins/Runtime/AudioSynesthesia/Source/AudioSynesthesia/Classes/AudioSynesthesiaNRT.h
new file mode 100644
--- /dev/null
+++ b/Engine/Plugins/Runtime/AudioSynesthesia/Source/AudioSynesthesia/Classes/AudioSynesthesiaNRT.h
@@ -0,0 +1,42 @@
+// Copyright 1998-2019 Epic Games, Inc. All Rights Reserved.
+#pragma once
+
+#include "AudioAnalyzerNRT.h"
+#include "AudioSynesthesiaNRT.generated.h"
+
+/** UAudioSynesthesiaNRTSettings
+ *
+ * Defines asset actions for derived UAudioSynthesiaNRTSettings subclasses.
+ */
+UCLASS(Abstract, Blueprintable)
+class AUDIOSYNESTHESIA_API UAudioSynesthesiaNRTSettings : public UAudioAnalyzerNRTSettings
+{
+	GENERATED_BODY()
+
+	public:
+
+		FText GetAssetActionName() const override;
+
+		UClass* GetSupportedClass() const override;
+
+		FColor GetTypeColor() const override;
+};
+
+/** UAudioSynesthesiaNRT
+ *
+ * Defines asset actions for derived UAudioSynthesiaNRT subclasses.
+ */
+UCLASS(Abstract, Blueprintable)
+class AUDIOSYNESTHESIA_API UAudioSynesthesiaNRT : public UAudioAnalyzerNRT
+{
+	GENERATED_BODY()
+
+	public:
+
+		FText GetAssetActionName() const override;
+
+		UClass* GetSupportedClass() const override;
+
+		FColor GetTypeColor() const override;
+};
+
diff --git a/Engine/Plugins/Runtime/AudioSynesthesia/Source/AudioSynesthesia/Private/AudioSynesthesiaModule.cpp b/Engine/Plugins/Runtime/AudioSynesthesia/Source/AudioSynesthesia/Private/AudioSynesthesiaModule.cpp
new file mode 100644
--- /dev/null
+++ b/Engine/Plugins/Runtime/AudioSynesthesia/Source/AudioSynesthesia/Private/AudioSynesthesiaModule.cpp
@@ -0,0 +1,23 @@
+// Copyright 1998-2019 Epic Games, Inc. All Rights Reserved.
+
+#include "AudioSynesthesiaModule.h"
+#include "CoreMinimal.h"
+#include "Modules/ModuleInterface.h"
+#include "Features/IModularFeatures.h"
+#include "LoudnessNRTFactory.h"
+
+DEFINE_LOG_CATEGORY(LogAudioSynesthesia);
+
+IMPLEMENT_MODULE(Audio::FAudioSynesthesiaModule, AudioSynesthesia)
+
+namespace Audio
+{
+	void FAudioSynesthesiaModule::StartupModule()
+	{
+		IModularFeatures::Get().RegisterModularFeature(Audio::FLoudnessNRTFactory::GetModularFeatureName(), &LoudnessNRTFactory);
+	}
+
+	void FAudioSynesthesiaModule::ShutdownModule()
+	{
+	}
+}
diff --git a/Engine/Plugins/Runtime/AudioSynesthesia/Source/AudioSynesthesia/Private/AudioSynesthesiaNRT.cpp b/Engine/Plugins/Runtime/AudioSynesthesia/Source/AudioSynesthesia/Private/AudioSynesthesiaNRT.cpp
new file mode 100644
--- /dev/null
+++ b/Engine/Plugins/Runtime/AudioSynesthesia/Source/AudioSynesthesia/Private/AudioSynesthesiaNRT.cpp
@@ -0,0 +1,35 @@
+// Copyright 1998-2019 Epic Games, Inc. All Rights Reserved.
+
+#include "AudioSynesthesiaNRT.h"
+
+
+FText UAudioSynesthesiaNRTSettings::GetAssetActionName() const
+{
+	return FText::FromString("Audio Synesthesia NRT Settings"); 
+}
+
+UClass* UAudioSynesthesiaNRTSettings::GetSupportedClass() const 
+{
+	return UAudioSynesthesiaNRTSettings::StaticClass();
+}
+
+FColor UAudioSynesthesiaNRTSettings::GetTypeColor() const
+{
+	return FColor(200.0f, 150.0f, 200.0f);
+}
+
+FText UAudioSynesthesiaNRT::GetAssetActionName() const
+{
+	return FText::FromString("Audio Synesthesia NRT"); 
+}
+
+UClass* UAudioSynesthesiaNRT::GetSupportedClass() const
+{
+	return UAudioSynesthesiaNRT::StaticClass();
+}
+
+FColor UAudioSynesthesiaNRT::GetTypeColor() const
+{
+	return FColor(200.0f, 150.0f, 200.0f);
+}
+
diff --git a/Engine/Plugins/Runtime/AudioSynesthesia/Source/AudioSynesthesia/Public/AudioSynesthesiaModule.h b/Engine/Plugins/Runtime/AudioSynesthesia/Source/AudioSynesthesia/Public/AudioSynesthesiaModule.h
new file mode 100644
--- /dev/null
+++ b/Engine/Plugins/Runtime/AudioSynesthesia/Source/AudioSynesthesia/Public/AudioSynesthesiaModule.h
@@ -0,0 +1,25 @@
+// Copyright 1998-2019 Epic Games, Inc. All Rights Reserved.
+
+#pragma once
+
+#include "Modules/ModuleInterface.h"
+#include "Modules/ModuleManager.h"
+#include "LoudnessNRTFactory.h"
+
+DECLARE_LOG_CATEGORY_EXTERN(LogAudioSynesthesia, Log, All);
+
+namespace Audio
+{
+	class AUDIOSYNESTHESIA_API FAudioSynesthesiaModule : public IModuleInterface
+	{
+	public:
+
+		// IModuleInterface interface
+		virtual void StartupModule() override;
+		virtual void ShutdownModule() override;
+
+	private:
+
+		FLoudnessNRTFactory LoudnessNRTFactory;
+	};
+}
diff --git a/Engine/Plugins/Runtime/AudioSynesthesia/Source/AudioSynesthesiaEditor/AudioSynesthesiaEditor.Build.cs b/Engine/Plugins/Runtime/AudioSynesthesia/Source/AudioSynesthesiaEditor/AudioSynesthesiaEditor.Build.cs
new file mode 100644
--- /dev/null
+++ b/Engine/Plugins/Runtime/AudioSynesthesia/Source/AudioSynesthesiaEditor/AudioSynesthesiaEditor.Build.cs
@@ -0,0 +1,31 @@
+// Copyright 1998-2019 Epic Games, Inc. All Rights Reserved.
+
+using UnrealBuildTool;
+
+public class AudioSynesthesiaEditor : ModuleRules
+{
+	public AudioSynesthesiaEditor(ReadOnlyTargetRules Target) : base(Target)
+	{
+
+		PublicDependencyModuleNames.AddRange(
+			new string[]
+			{
+				"Core",
+				"CoreUObject",
+				"Engine",
+				"UnrealEd",
+				"AudioSynesthesia",
+				"AudioAnalyzer",
+                "InputCore",
+                "Slate",
+                "SlateCore",
+                "EditorStyle"
+            }
+        );
+
+		PrivateIncludePathModuleNames.AddRange(
+			new string[] {
+				"AssetTools"
+			});
+	}
+}
diff --git a/Engine/Plugins/Runtime/AudioSynesthesia/Source/AudioSynesthesiaEditor/Classes/AudioSynesthesiaNRTFactory.h b/Engine/Plugins/Runtime/AudioSynesthesia/Source/AudioSynesthesiaEditor/Classes/AudioSynesthesiaNRTFactory.h
new file mode 100644
--- /dev/null
+++ b/Engine/Plugins/Runtime/AudioSynesthesia/Source/AudioSynesthesiaEditor/Classes/AudioSynesthesiaNRTFactory.h
@@ -0,0 +1,27 @@
+// Copyright 1998-2019 Epic Games, Inc. All Rights Reserved.
+
+#pragma once
+
+#include "AudioSynesthesiaNRT.h"
+#include "Factories/Factory.h"
+#include "AudioSynesthesiaNRTFactory.generated.h"
+
+/** UAudioSynesthesiaNRTFactory
+ *
+ * UAudioSynesthesiaNRTFactory implements the factory for UAudioSynesthesiaNRT assets.
+ */
+UCLASS(MinimalAPI, hidecategories=Object)
+class UAudioSynesthesiaNRTFactory : public UFactory
+{
+	GENERATED_UCLASS_BODY()
+
+	/** The type of audio analyzer settings that will be created */
+	UPROPERTY(EditAnywhere, Category = AudioSynesthesiaFactory)
+	TSubclassOf<UAudioSynesthesiaNRT> AudioSynesthesiaNRTClass;
+
+	//~ Begin UFactory Interface
+	virtual bool ConfigureProperties() override;
+	virtual UObject* FactoryCreateNew(UClass* Class,UObject* InParent,FName Name,EObjectFlags Flags,UObject* Context, FFeedbackContext* Warn) override;
+	//~ Begin UFactory Interface	
+};
+
diff --git a/Engine/Plugins/Runtime/AudioSynesthesia/Source/AudioSynesthesiaEditor/Classes/AudioSynesthesiaNRTSettingsFactory.h b/Engine/Plugins/Runtime/AudioSynesthesia/Source/AudioSynesthesiaEditor/Classes/AudioSynesthesiaNRTSettingsFactory.h
new file mode 100644
--- /dev/null
+++ b/Engine/Plugins/Runtime/AudioSynesthesia/Source/AudioSynesthesiaEditor/Classes/AudioSynesthesiaNRTSettingsFactory.h
@@ -0,0 +1,28 @@
+// Copyright 1998-2019 Epic Games, Inc. All Rights Reserved.
+
+#pragma once
+
+#include "AudioSynesthesiaNRT.h"
+#include "Factories/Factory.h"
+#include "AudioSynesthesiaNRTSettingsFactory.generated.h"
+
+/** UAudioSynesthesiaNRTSettingsFactory
+ *
+ * UAudioSynesthesiaNRTSettingsFactory implements the factory for UAudioSynesthesiaNRTSettings assets.
+ */
+UCLASS(MinimalAPI, hidecategories=Object)
+class UAudioSynesthesiaNRTSettingsFactory : public UFactory
+{
+	GENERATED_UCLASS_BODY()
+
+	/** The type of audio analyzer settings that will be created */
+	UPROPERTY(EditAnywhere, Category = AudioSynesthesiaFactory)
+	TSubclassOf<UAudioSynesthesiaNRTSettings> AudioSynesthesiaNRTSettingsClass;
+
+	//~ Begin UFactory Interface
+	virtual bool ConfigureProperties() override;
+	virtual UObject* FactoryCreateNew(UClass* Class,UObject* InParent,FName Name,EObjectFlags Flags,UObject* Context, FFeedbackContext* Warn) override;
+	//~ Begin UFactory Interface	
+};
+
+
diff --git a/Engine/Plugins/Runtime/AudioSynesthesia/Source/AudioSynesthesiaEditor/Private/AssetTypeActions_AudioSynesthesiaNRT.cpp b/Engine/Plugins/Runtime/AudioSynesthesia/Source/AudioSynesthesiaEditor/Private/AssetTypeActions_AudioSynesthesiaNRT.cpp
new file mode 100644
--- /dev/null
+++ b/Engine/Plugins/Runtime/AudioSynesthesia/Source/AudioSynesthesiaEditor/Private/AssetTypeActions_AudioSynesthesiaNRT.cpp
@@ -0,0 +1,49 @@
+// Copyright 1998-2019 Epic Games, Inc. All Rights Reserved.
+
+#include "AssetTypeActions_AudioSynesthesiaNRT.h"
+#include "Developer/AssetTools/Public/AssetTypeCategories.h"
+#include "AudioSynesthesiaNRT.h"
+
+FAssetTypeActions_AudioSynesthesiaNRT::FAssetTypeActions_AudioSynesthesiaNRT(UAudioSynesthesiaNRT* InSynesthesia)
+	: Synesthesia(InSynesthesia)
+{
+}
+
+FText FAssetTypeActions_AudioSynesthesiaNRT::GetName() const
+{
+	FText AssetActionName = Synesthesia->GetAssetActionName();
+	if (AssetActionName.IsEmpty())
+	{
+		FString ClassName;
+		Synesthesia->GetClass()->GetName(ClassName);
+		return FText::FromString(ClassName);
+	}
+	else
+	{
+		return AssetActionName;
+	}
+}
+
+FColor FAssetTypeActions_AudioSynesthesiaNRT::GetTypeColor() const 
+{
+	return Synesthesia->GetTypeColor(); 
+}
+
+UClass* FAssetTypeActions_AudioSynesthesiaNRT::GetSupportedClass() const
+{
+	UClass* SupportedClass = Synesthesia->GetSupportedClass();
+	if (SupportedClass == nullptr)
+	{
+		return Synesthesia->GetClass();
+	}
+	else
+	{
+		return SupportedClass;
+	}
+}
+
+uint32 FAssetTypeActions_AudioSynesthesiaNRT::GetCategories() 
+{
+	return EAssetTypeCategories::Sounds; 
+}
+
diff --git a/Engine/Plugins/Runtime/AudioSynesthesia/Source/AudioSynesthesiaEditor/Private/AssetTypeActions_AudioSynesthesiaNRT.h b/Engine/Plugins/Runtime/AudioSynesthesia/Source/AudioSynesthesiaEditor/Private/AssetTypeActions_AudioSynesthesiaNRT.h
new file mode 100644
--- /dev/null
+++ b/Engine/Plugins/Runtime/AudioSynesthesia/Source/AudioSynesthesiaEditor/Private/AssetTypeActions_AudioSynesthesiaNRT.h
@@ -0,0 +1,25 @@
+// Copyright 1998-2019 Epic Games, Inc. All Rights Reserved.
+
+#pragma once
+
+#include "CoreMinimal.h"
+#include "AssetTypeActions_Base.h"
+
+class UAudioSynesthesiaNRT;
+
+/** FAssetTypeActions_AudioSynesthesiaNRT */
+class FAssetTypeActions_AudioSynesthesiaNRT : public FAssetTypeActions_Base
+{
+public:
+	FAssetTypeActions_AudioSynesthesiaNRT(UAudioSynesthesiaNRT* InSynesthesia);
+
+	//~ Begin FAssetTypeActions_Base
+	virtual FText GetName() const override;
+	virtual FColor GetTypeColor() const override;
+	virtual UClass* GetSupportedClass() const override;
+	virtual uint32 GetCategories() override;
+	//~ End FAssetTypeActions_Base
+
+private:
+	UAudioSynesthesiaNRT* Synesthesia;
+};
diff --git a/Engine/Plugins/Runtime/AudioSynesthesia/Source/AudioSynesthesiaEditor/Private/AssetTypeActions_AudioSynesthesiaNRTSettings.cpp b/Engine/Plugins/Runtime/AudioSynesthesia/Source/AudioSynesthesiaEditor/Private/AssetTypeActions_AudioSynesthesiaNRTSettings.cpp
new file mode 100644
--- /dev/null
+++ b/Engine/Plugins/Runtime/AudioSynesthesia/Source/AudioSynesthesiaEditor/Private/AssetTypeActions_AudioSynesthesiaNRTSettings.cpp
@@ -0,0 +1,49 @@
+// Copyright 1998-2019 Epic Games, Inc. All Rights Reserved.
+
+#include "AssetTypeActions_AudioSynesthesiaNRTSettings.h"
+#include "Developer/AssetTools/Public/AssetTypeCategories.h"
+#include "AudioSynesthesiaNRT.h"
+
+FAssetTypeActions_AudioSynesthesiaNRTSettings::FAssetTypeActions_AudioSynesthesiaNRTSettings(UAudioSynesthesiaNRTSettings* InSynesthesiaSettings)
+	: SynesthesiaSettings(InSynesthesiaSettings)
+{
+}
+
+FText FAssetTypeActions_AudioSynesthesiaNRTSettings::GetName() const
+{
+	FText AssetActionName = SynesthesiaSettings->GetAssetActionName();
+	if (AssetActionName.IsEmpty())
+	{
+		FString ClassName;
+		SynesthesiaSettings->GetClass()->GetName(ClassName);
+		return FText::FromString(ClassName);
+	}
+	else
+	{
+		return AssetActionName;
+	}
+}
+
+FColor FAssetTypeActions_AudioSynesthesiaNRTSettings::GetTypeColor() const 
+{
+	return SynesthesiaSettings->GetTypeColor(); 
+}
+
+UClass* FAssetTypeActions_AudioSynesthesiaNRTSettings::GetSupportedClass() const
+{
+	UClass* SupportedClass = SynesthesiaSettings->GetSupportedClass();
+	if (SupportedClass == nullptr)
+	{
+		return SynesthesiaSettings->GetClass();
+	}
+	else
+	{
+		return SupportedClass;
+	}
+}
+
+uint32 FAssetTypeActions_AudioSynesthesiaNRTSettings::GetCategories() 
+{
+	return EAssetTypeCategories::Sounds; 
+}
+
diff --git a/Engine/Plugins/Runtime/AudioSynesthesia/Source/AudioSynesthesiaEditor/Private/AssetTypeActions_AudioSynesthesiaNRTSettings.h b/Engine/Plugins/Runtime/AudioSynesthesia/Source/AudioSynesthesiaEditor/Private/AssetTypeActions_AudioSynesthesiaNRTSettings.h
new file mode 100644
--- /dev/null
+++ b/Engine/Plugins/Runtime/AudioSynesthesia/Source/AudioSynesthesiaEditor/Private/AssetTypeActions_AudioSynesthesiaNRTSettings.h
@@ -0,0 +1,25 @@
+// Copyright 1998-2019 Epic Games, Inc. All Rights Reserved.
+
+#pragma once
+
+#include "CoreMinimal.h"
+#include "AssetTypeActions_Base.h"
+
+class UAudioSynesthesiaNRTSettings;
+
+/** FAssetTypeActions_AudioSynesthesiaNRTSettings */
+class FAssetTypeActions_AudioSynesthesiaNRTSettings : public FAssetTypeActions_Base
+{
+public:
+	FAssetTypeActions_AudioSynesthesiaNRTSettings(UAudioSynesthesiaNRTSettings* InSynesthesiaSettings);
+
+	//~ Begin FAssetTypeActions_Base
+	virtual FText GetName() const override;
+	virtual FColor GetTypeColor() const override;
+	virtual UClass* GetSupportedClass() const override;
+	virtual uint32 GetCategories() override;
+	//~ End FAssetTypeActions_Base
+
+private:
+	UAudioSynesthesiaNRTSettings* SynesthesiaSettings;
+};
diff --git a/Engine/Plugins/Runtime/AudioSynesthesia/Source/AudioSynesthesiaEditor/Private/AudioSynesthesiaClassFilter.cpp b/Engine/Plugins/Runtime/AudioSynesthesia/Source/AudioSynesthesiaEditor/Private/AudioSynesthesiaClassFilter.cpp
new file mode 100644
--- /dev/null
+++ b/Engine/Plugins/Runtime/AudioSynesthesia/Source/AudioSynesthesiaEditor/Private/AudioSynesthesiaClassFilter.cpp
@@ -0,0 +1,18 @@
+// Copyright 1998-2019 Epic Games, Inc. All Rights Reserved.
+
+#include "AudioSynesthesiaClassFilter.h"
+
+namespace Audio
+{
+	bool FAssetClassParentFilter::IsClassAllowed(const FClassViewerInitializationOptions& InInitOptions, const UClass* InClass, TSharedRef< FClassViewerFilterFuncs > InFilterFuncs) 
+	{
+		return !InClass->HasAnyClassFlags(DisallowedClassFlags)
+			&& InFilterFuncs->IfInChildOfClassesSet(AllowedChildrenOfClasses, InClass) != EFilterReturn::Failed;
+	}
+
+	bool FAssetClassParentFilter::IsUnloadedClassAllowed(const FClassViewerInitializationOptions& InInitOptions, const TSharedRef< const IUnloadedBlueprintData > InUnloadedClassData, TSharedRef< FClassViewerFilterFuncs > InFilterFuncs) 
+	{
+		return !InUnloadedClassData->HasAnyClassFlags(DisallowedClassFlags)
+			&& InFilterFuncs->IfInChildOfClassesSet(AllowedChildrenOfClasses, InUnloadedClassData) != EFilterReturn::Failed;
+	}
+}
diff --git a/Engine/Plugins/Runtime/AudioSynesthesia/Source/AudioSynesthesiaEditor/Private/AudioSynesthesiaClassFilter.h b/Engine/Plugins/Runtime/AudioSynesthesia/Source/AudioSynesthesiaEditor/Private/AudioSynesthesiaClassFilter.h
new file mode 100644
--- /dev/null
+++ b/Engine/Plugins/Runtime/AudioSynesthesia/Source/AudioSynesthesiaEditor/Private/AudioSynesthesiaClassFilter.h
@@ -0,0 +1,22 @@
+// Copyright 1998-2019 Epic Games, Inc. All Rights Reserved.
+
+#pragma once
+
+#include "ClassViewerFilter.h"
+
+namespace Audio
+{
+	class FAssetClassParentFilter : public IClassViewerFilter
+	{
+	public:
+		/** All children of these classes will be included unless filtered out by another setting. */
+		TSet< const UClass* > AllowedChildrenOfClasses;
+
+		/** Disallowed class flags. */
+		EClassFlags DisallowedClassFlags;
+
+		virtual bool IsClassAllowed(const FClassViewerInitializationOptions& InInitOptions, const UClass* InClass, TSharedRef< FClassViewerFilterFuncs > InFilterFuncs) override;
+
+		virtual bool IsUnloadedClassAllowed(const FClassViewerInitializationOptions& InInitOptions, const TSharedRef< const IUnloadedBlueprintData > InUnloadedClassData, TSharedRef< FClassViewerFilterFuncs > InFilterFuncs) override;
+	};
+}
diff --git a/Engine/Plugins/Runtime/AudioSynesthesia/Source/AudioSynesthesiaEditor/Private/AudioSynesthesiaEditorModule.cpp b/Engine/Plugins/Runtime/AudioSynesthesia/Source/AudioSynesthesiaEditor/Private/AudioSynesthesiaEditorModule.cpp
new file mode 100644
--- /dev/null
+++ b/Engine/Plugins/Runtime/AudioSynesthesia/Source/AudioSynesthesiaEditor/Private/AudioSynesthesiaEditorModule.cpp
@@ -0,0 +1,68 @@
+// Copyright 1998-2019 Epic Games, Inc. All Rights Reserved.
+
+#include "AudioSynesthesiaEditorModule.h"
+#include "Modules/ModuleManager.h"
+#include "AssetToolsModule.h"
+#include "UObject/UObjectIterator.h"
+#include "AudioSynesthesiaNRT.h"
+#include "AssetTypeActions_AudioSynesthesiaNRT.h"
+#include "AssetTypeActions_AudioSynesthesiaNRTSettings.h"
+
+DEFINE_LOG_CATEGORY(LogAudioSynesthesiaEditor);
+
+class FAudioSynesthesiaEditorModule : public IAudioSynesthesiaEditorModule
+{
+public:
+	FAudioSynesthesiaEditorModule()
+	{}
+
+	virtual void StartupModule() override
+	{
+		RegisterAssetActions();
+	}
+
+	virtual void ShutdownModule() override
+	{}
+
+
+	virtual void RegisterAssetActions() override
+	{
+		RegisterAudioSynesthesiaAssetActions<UAudioSynesthesiaNRT, FAssetTypeActions_AudioSynesthesiaNRT>();
+		RegisterAudioSynesthesiaAssetActions<UAudioSynesthesiaNRTSettings, FAssetTypeActions_AudioSynesthesiaNRTSettings>();
+	}
+
+private:
+
+	template<class SynesthesiaAssetType, class SynesthesiaAssetActionsType >
+	void RegisterAudioSynesthesiaAssetActions() 
+	{
+		// Register the audio editor asset type actions
+		IAssetTools& AssetTools = FModuleManager::LoadModuleChecked<FAssetToolsModule>("AssetTools").Get();
+
+		// Look for any sound effect presets to register
+		for (TObjectIterator<UClass> It; It; ++It)
+		{
+			UClass* ChildClass = *It;
+			if (ChildClass->HasAnyClassFlags(CLASS_Abstract))
+			{
+				continue;
+			}
+
+			// look for Synesthesia classes 
+			UClass* ParentClass = ChildClass->GetSuperClass();
+			if (ParentClass->IsChildOf(SynesthesiaAssetType::StaticClass()))
+			{
+				SynesthesiaAssetType* Synesthesia = ChildClass->GetDefaultObject<SynesthesiaAssetType>();
+				if (!RegisteredActions.Contains(Synesthesia) && Synesthesia->HasAssetActions())
+				{
+					RegisteredActions.Add(Synesthesia);
+					AssetTools.RegisterAssetTypeActions(MakeShareable(new SynesthesiaAssetActionsType(Synesthesia)));
+				}
+			}
+		}
+	}	
+
+	TSet<UObject*> RegisteredActions;
+};
+
+IMPLEMENT_MODULE( FAudioSynesthesiaEditorModule, AudioSynesthesiaEditor );
diff --git a/Engine/Plugins/Runtime/AudioSynesthesia/Source/AudioSynesthesiaEditor/Private/AudioSynesthesiaNRTFactory.cpp b/Engine/Plugins/Runtime/AudioSynesthesia/Source/AudioSynesthesiaEditor/Private/AudioSynesthesiaNRTFactory.cpp
new file mode 100644
--- /dev/null
+++ b/Engine/Plugins/Runtime/AudioSynesthesia/Source/AudioSynesthesiaEditor/Private/AudioSynesthesiaNRTFactory.cpp
@@ -0,0 +1,62 @@
+// Copyright 1998-2019 Epic Games, Inc. All Rights Reserved.
+
+#include "AudioSynesthesiaNRTFactory.h"
+
+#include "ClassViewerModule.h"
+#include "ClassViewerFilter.h"
+#include "Kismet2/SClassPickerDialog.h"
+#include "Modules/ModuleManager.h"
+#include "AudioSynesthesiaClassFilter.h"
+
+#define LOCTEXT_NAMESPACE "AudioSynesthesiaEditor"
+
+UAudioSynesthesiaNRTFactory::UAudioSynesthesiaNRTFactory(const FObjectInitializer& ObjectInitializer)
+	: Super(ObjectInitializer)
+{
+	SupportedClass = UAudioSynesthesiaNRT::StaticClass();
+
+	bCreateNew = true;
+	bEditorImport = false;
+	bEditAfterNew = true;
+	AudioSynesthesiaNRTClass = nullptr;
+}
+
+bool UAudioSynesthesiaNRTFactory::ConfigureProperties()
+{
+	AudioSynesthesiaNRTClass = nullptr;
+
+	// Load the classviewer module to display a class picker
+	FClassViewerModule& ClassViewerModule = FModuleManager::LoadModuleChecked<FClassViewerModule>("ClassViewer");
+
+	FClassViewerInitializationOptions Options;
+	Options.Mode = EClassViewerMode::ClassPicker;
+
+	TSharedPtr<Audio::FAssetClassParentFilter> Filter = MakeShareable(new Audio::FAssetClassParentFilter);
+	Options.ClassFilter = Filter;
+
+	Filter->DisallowedClassFlags = CLASS_Abstract | CLASS_Deprecated | CLASS_NewerVersionExists;
+	Filter->AllowedChildrenOfClasses.Add(UAudioSynesthesiaNRT::StaticClass());
+
+	const FText TitleText = LOCTEXT("CreateAudioSynesthesiaNRTOptions", "Pick Synesthesia Class");
+	UClass* ChosenClass = nullptr;
+	const bool bPressedOk = SClassPickerDialog::PickClass(TitleText, Options, ChosenClass, UAudioSynesthesiaNRT::StaticClass());
+
+	if (bPressedOk)
+	{
+		AudioSynesthesiaNRTClass = ChosenClass;
+	}
+
+	return bPressedOk;
+}
+
+UObject* UAudioSynesthesiaNRTFactory::FactoryCreateNew(UClass* InClass, UObject* InParent, FName InName, EObjectFlags Flags, UObject* Context, FFeedbackContext* Warn)
+{
+	UAudioSynesthesiaNRT* NewAudioSynesthesiaNRT = nullptr;
+	if (AudioSynesthesiaNRTClass != nullptr)
+	{
+		NewAudioSynesthesiaNRT = NewObject<UAudioSynesthesiaNRT>(InParent, AudioSynesthesiaNRTClass, InName, Flags);
+	}
+	return NewAudioSynesthesiaNRT;
+}
+
+#undef LOCTEXT_NAMESPACE
diff --git a/Engine/Plugins/Runtime/AudioSynesthesia/Source/AudioSynesthesiaEditor/Private/AudioSynesthesiaNRTSettingsFactory.cpp b/Engine/Plugins/Runtime/AudioSynesthesia/Source/AudioSynesthesiaEditor/Private/AudioSynesthesiaNRTSettingsFactory.cpp
new file mode 100644
--- /dev/null
+++ b/Engine/Plugins/Runtime/AudioSynesthesia/Source/AudioSynesthesiaEditor/Private/AudioSynesthesiaNRTSettingsFactory.cpp
@@ -0,0 +1,61 @@
+// Copyright 1998-2019 Epic Games, Inc. All Rights Reserved.
+
+#include "AudioSynesthesiaNRTSettingsFactory.h"
+
+#include "ClassViewerModule.h"
+#include "ClassViewerFilter.h"
+#include "Kismet2/SClassPickerDialog.h"
+#include "Modules/ModuleManager.h"
+#include "AudioSynesthesiaClassFilter.h"
+
+UAudioSynesthesiaNRTSettingsFactory::UAudioSynesthesiaNRTSettingsFactory(const FObjectInitializer& ObjectInitializer)
+	: Super(ObjectInitializer)
+{
+	SupportedClass = UAudioSynesthesiaNRTSettings::StaticClass();
+
+	bCreateNew = true;
+	bEditorImport = false;
+	bEditAfterNew = true;
+	AudioSynesthesiaNRTSettingsClass = nullptr;
+}
+#define LOCTEXT_NAMESPACE "AudioSynesthesiaEditor"
+
+bool UAudioSynesthesiaNRTSettingsFactory::ConfigureProperties()
+{
+	AudioSynesthesiaNRTSettingsClass = nullptr;
+
+	// Load the classviewer module to display a class picker
+	FClassViewerModule& ClassViewerModule = FModuleManager::LoadModuleChecked<FClassViewerModule>("ClassViewer");
+
+	FClassViewerInitializationOptions Options;
+	Options.Mode = EClassViewerMode::ClassPicker;
+
+	TSharedPtr<Audio::FAssetClassParentFilter> Filter = MakeShareable(new Audio::FAssetClassParentFilter);
+	Options.ClassFilter = Filter;
+
+	Filter->DisallowedClassFlags = CLASS_Abstract | CLASS_Deprecated | CLASS_NewerVersionExists;
+	Filter->AllowedChildrenOfClasses.Add(UAudioSynesthesiaNRTSettings::StaticClass());
+
+	const FText TitleText = LOCTEXT("CreateAudioSynesthesiaNRTSettingOptions", "Pick Synesthesia Settings Class");
+	UClass* ChosenClass = nullptr;
+	const bool bPressedOk = SClassPickerDialog::PickClass(TitleText, Options, ChosenClass, UAudioSynesthesiaNRTSettings::StaticClass());
+
+	if (bPressedOk)
+	{
+		AudioSynesthesiaNRTSettingsClass = ChosenClass;
+	}
+
+	return bPressedOk;
+}
+
+UObject* UAudioSynesthesiaNRTSettingsFactory::FactoryCreateNew(UClass* InClass, UObject* InParent, FName InName, EObjectFlags Flags, UObject* Context, FFeedbackContext* Warn)
+{
+	UAudioSynesthesiaNRTSettings* NewAudioSynesthesiaNRTSettings = nullptr;
+	if (AudioSynesthesiaNRTSettingsClass != nullptr)
+	{
+		NewAudioSynesthesiaNRTSettings = NewObject<UAudioSynesthesiaNRTSettings>(InParent, AudioSynesthesiaNRTSettingsClass, InName, Flags);
+	}
+	return NewAudioSynesthesiaNRTSettings;
+}
+
+#undef LOCTEXT_NAMESPACE
diff --git a/Engine/Plugins/Runtime/AudioSynesthesia/Source/AudioSynesthesiaEditor/Public/AudioSynesthesiaEditorModule.h b/Engine/Plugins/Runtime/AudioSynesthesia/Source/AudioSynesthesiaEditor/Public/AudioSynesthesiaEditorModule.h
new file mode 100644
--- /dev/null
+++ b/Engine/Plugins/Runtime/AudioSynesthesia/Source/AudioSynesthesiaEditor/Public/AudioSynesthesiaEditorModule.h
@@ -0,0 +1,18 @@
+// Copyright 1998-2019 Epic Games, Inc. All Rights Reserved.
+
+#pragma once
+
+#include "CoreMinimal.h"
+#include "Modules/ModuleInterface.h"
+
+AUDIOSYNESTHESIAEDITOR_API DECLARE_LOG_CATEGORY_EXTERN(LogAudioSynesthesiaEditor, Log, All);
+
+extern const FName AudioSynesthesiaEditorAppIdentifier;
+
+/** audio analyzer editor module interface */
+class AUDIOSYNESTHESIAEDITOR_API IAudioSynesthesiaEditorModule :	public IModuleInterface
+{
+public:
+	/** Registers audio analyzer asset actions. */
+	virtual void RegisterAssetActions() = 0;
+};
diff --git a/Engine/Source/Runtime/AudioAnalyzer/AudioAnalyzer.Build.cs b/Engine/Source/Runtime/AudioAnalyzer/AudioAnalyzer.Build.cs
new file mode 100644
--- /dev/null
+++ b/Engine/Source/Runtime/AudioAnalyzer/AudioAnalyzer.Build.cs
@@ -0,0 +1,29 @@
+// Copyright 1998-2019 Epic Games, Inc. All Rights Reserved.
+// Copyright 1998-2018 Epic Games, Inc. All Rights Reserved.
+
+namespace UnrealBuildTool.Rules
+{
+	public class AudioAnalyzer : ModuleRules
+	{
+		public AudioAnalyzer(ReadOnlyTargetRules Target) : base(Target)
+		{
+
+            PrivateDependencyModuleNames.AddRange(
+				new string[]
+				{
+                    "Core",
+                    "InputCore",
+                    "CoreUObject",
+                    "Engine"
+                }
+			);
+
+            if (Target.Type == TargetType.Editor &&
+				(Target.Platform == UnrealTargetPlatform.Win64 || Target.Platform == UnrealTargetPlatform.Win32)
+                )
+            {
+                AddEngineThirdPartyPrivateStaticDependencies(Target, "UELibSampleRate");
+            }
+        }
+	}
+}
diff --git a/Engine/Source/Runtime/AudioAnalyzer/Classes/AudioAnalyzerAsset.h b/Engine/Source/Runtime/AudioAnalyzer/Classes/AudioAnalyzerAsset.h
new file mode 100644
--- /dev/null
+++ b/Engine/Source/Runtime/AudioAnalyzer/Classes/AudioAnalyzerAsset.h
@@ -0,0 +1,41 @@
+// Copyright 1998-2019 Epic Games, Inc. All Rights Reserved.
+
+/** Interface for Audio Analyzer Assets. */
+
+#pragma once
+
+#include "CoreMinimal.h"
+#include "AudioAnalyzerAsset.generated.h"
+
+/** UAudioAnalyzerAsset
+ *
+ * UAudioAnalyzerAsset provides the interface for controlling asset actions within the editor.
+ */
+UCLASS(Abstract, EditInlineNew)
+class AUDIOANALYZER_API UAudioAnalyzerAsset : public UObject
+{
+	GENERATED_BODY()
+
+	public:
+
+#ifdef WITH_EDITOR
+		virtual bool HasAssetActions() const { return true; }
+
+		/**
+		 * GetAssetActionName() returns the FText displayed in the editor.
+		 */ 
+		virtual FText GetAssetActionName() const PURE_VIRTUAL(UAudioAnalyzerNRTAsset::GetAssetActionName, return FText(););
+
+		/**
+		 * GetSupportedClass() returns the class which should be associated with these asset actions.
+		 */ 
+		virtual UClass* GetSupportedClass() const PURE_VIRTUAL(UAudioAnalyzerNRTAsset::GetSupportedClass, return nullptr;);
+
+		/**
+		 * GetTypeColor() returns the color used to display asset icons within the editor.
+		 */
+		virtual FColor GetTypeColor() const { return FColor(100.0f, 100.0f, 100.0f); }
+#endif
+
+};
+
diff --git a/Engine/Source/Runtime/AudioAnalyzer/Classes/AudioAnalyzerNRT.h b/Engine/Source/Runtime/AudioAnalyzer/Classes/AudioAnalyzerNRT.h
new file mode 100644
--- /dev/null
+++ b/Engine/Source/Runtime/AudioAnalyzer/Classes/AudioAnalyzerNRT.h
@@ -0,0 +1,147 @@
+// Copyright 1998-2019 Epic Games, Inc. All Rights Reserved.
+
+/** Interface for Non Real-Time (NRT) Audio Analyzer UObjects. */
+
+#pragma once
+
+#include "CoreMinimal.h"
+#include "IAudioAnalyzerNRTInterface.h"
+#include "AudioAnalyzerAsset.h"
+#include "Sound/SoundWave.h"
+#include "AudioAnalyzerNRT.generated.h"
+
+DECLARE_DELEGATE( FAnalyzeAudioDelegate );
+
+/** UAudioAnalyzerNRTSettings
+ *
+ * UAudioAnalyzerNRTSettings provides a way to store and reuse existing analyzer settings
+ * across multipler analyzers. This class provides the interface and functionality to 
+ * automatically trigger reanalysis of audio across all analyzers associated with this 
+ * settingwhen when a UPROPERTY in this setting object is edited.
+ *
+ */
+UCLASS(Abstract, EditInlineNew, BlueprintType)
+class AUDIOANALYZER_API UAudioAnalyzerNRTSettings : public UAudioAnalyzerAsset
+{
+	GENERATED_BODY()
+
+	public:
+
+#if WITH_EDITOR
+
+		FAnalyzeAudioDelegate AnalyzeAudioDelegate;
+
+		/** 
+		 * Called when a UPROPERTY of this class is edited. Triggers
+		 * reanalysis of audio.
+		 *
+		 * This determines whether to trigger analysis by calling ShouldEventTriggerAnalysis(...)
+		 */
+		void PostEditChangeProperty (struct FPropertyChangedEvent & PropertyChangedEvent) override;
+
+		/** 
+		 * This returns true when the PropertyChangeEvent is due to update a setting property. 
+		 * Override this method in order to customize this behavior.
+		 */
+		virtual bool ShouldEventTriggerAnalysis(struct FPropertyChangedEvent & PropertyChangeEvent);
+#endif
+
+};
+
+/** UAudioAnalyzerNRT
+ *
+ * UAudioAnalyzerNRT applies an analyzer to a sound using specific settings, stores the 
+ * results and exposes them via blueprints. 
+ *
+ * Subclasses of UAudioAnalyzerNRT must implement GetAnalyzerNRTFactoryName() to associate
+ * the UAudioAnalyzerNRT with an IAudioAnalyzerNRTFactory implementation. 
+ *
+ * To support blueprint access, subclasses can implement UFUNCTIONs to expose the data
+ * returned by GetResult().
+ */
+UCLASS(Abstract, EditInlineNew, BlueprintType)
+class AUDIOANALYZER_API UAudioAnalyzerNRT : public UAudioAnalyzerAsset
+{
+	GENERATED_BODY()
+
+	public:
+
+		/**
+		 * The USoundWave which is analyzed.
+		 */
+		UPROPERTY(EditAnywhere, BlueprintReadOnly, Category=AudioAnalyzer)
+		USoundWave* Sound;
+
+		/**
+		 * Returns the result object generated by the associated IAudioAnalyzerNRTFactory.
+		 * The template argument must be the IAudioAnalyzerNRTResult subclass returned by
+		 * the associated IAudioAnalyzerNRTFactory.
+		 */
+		template<class ResultType>
+		TSharedPtr<ResultType, ESPMode::ThreadSafe> GetResult()
+		{
+			TSharedPtr<ResultType, ESPMode::ThreadSafe> ReturnedResult;
+			{
+				FScopeLock ResultLock(&ResultCriticalSection);
+				ReturnedResult = StaticCastSharedPtr<ResultType>(Result); 
+			}
+			return ReturnedResult;
+		}
+	 
+		/** 
+		 * Implementations can override this method to create settings objects
+		 * specific for their analyzer. 
+		 */
+		virtual TUniquePtr<Audio::IAnalyzerNRTSettings> GetSettings();
+
+		/**
+		 * Performs serialization of results.
+		 */
+		virtual void Serialize(FArchive& Ar) override;
+
+		/* TODO: progress/status for calculation */
+
+#if WITH_EDITOR
+
+		void SetResult(TSharedPtr<Audio::IAnalyzerNRTResult, ESPMode::ThreadSafe> NewResult);
+
+		/**
+		 * Called before a UPROPERTY of this class is edited. This checks to see
+		 * if the UPROPERTY is a UAudioAnalyzerNRTSettings. If so, the previous settings' 
+		 * AnalyzeAudioDelegate will be unbound from this object's AnalyzeAudio()
+		 */
+		void PreEditChange(UProperty* PropertyAboutToChange) override;
+
+		/**
+		 * Called when a UPROPERTY of this class is edited. Triggering 
+		 * reanalysis of audio when appropriate and binds new settings if the UPROPERTY
+		 * is a UAudioAnalyzerNRTSettings derived object.
+		 *
+		 * This determines whether to trigger analysis by calling ShouldEventTriggerAnalysis(...)
+		 */
+		void PostEditChangeProperty (struct FPropertyChangedEvent & PropertyChangedEvent) override;
+
+		/**
+		 * Returns true when the PropertyChangeEvent is due to update SoundWave or Settings.
+		 */
+		virtual bool ShouldEventTriggerAnalysis(struct FPropertyChangedEvent & PropertyChangeEvent);
+
+		/** Performs the analaysis of the audio */
+		UFUNCTION()
+		void AnalyzeAudio();
+#endif
+
+
+	protected:
+		/* Subclasses must override this method in order to inform this object which AnalyzerNRTFactory to use for analysis */
+		virtual FName GetAnalyzerNRTFactoryName() const PURE_VIRTUAL(UAudioAnalyzerNRT::GetAnalyzerNRTFactoryName, return FName(););
+
+	private:
+
+		// Returns UAudioAnalyzerNRTSettings* if property points to a valid UAudioAnalyzerNRTSettings, otherwise returns nullptr.
+		UAudioAnalyzerNRTSettings* GetSettingsFromProperty(UProperty* Property);
+
+		TSharedPtr<Audio::IAnalyzerNRTResult, ESPMode::ThreadSafe> Result;
+		FCriticalSection ResultCriticalSection;
+};
+
diff --git a/Engine/Source/Runtime/AudioAnalyzer/Private/AudioAnalyzerModule.cpp b/Engine/Source/Runtime/AudioAnalyzer/Private/AudioAnalyzerModule.cpp
new file mode 100644
--- /dev/null
+++ b/Engine/Source/Runtime/AudioAnalyzer/Private/AudioAnalyzerModule.cpp
@@ -0,0 +1,15 @@
+// Copyright 1998-2019 Epic Games, Inc. All Rights Reserved.
+#include "AudioAnalyzerModule.h"
+
+IMPLEMENT_MODULE(FAudioAnalyzerModule, AudioAnalyzer);
+
+DEFINE_LOG_CATEGORY(LogAudioAnalyzer);
+
+void FAudioAnalyzerModule::StartupModule()
+{
+	int x = 32;
+}
+
+void FAudioAnalyzerModule::ShutdownModule()
+{
+}
diff --git a/Engine/Source/Runtime/AudioAnalyzer/Private/AudioAnalyzerNRT.cpp b/Engine/Source/Runtime/AudioAnalyzer/Private/AudioAnalyzerNRT.cpp
new file mode 100644
--- /dev/null
+++ b/Engine/Source/Runtime/AudioAnalyzer/Private/AudioAnalyzerNRT.cpp
@@ -0,0 +1,249 @@
+// Copyright 1998-2019 Epic Games, Inc. All Rights Reserved.
+
+#include "AudioAnalyzerNRT.h"
+#include "AudioAnalyzerNRTFacade.h"
+#include "Sound/SampleBuffer.h"
+
+
+#if WITH_EDITOR
+
+namespace 
+{
+	class FAudioAnalyzeTask : public FNonAbandonableTask
+	{
+		friend class FAutoDeleteAsyncTask<FAudioAnalyzeTask>;
+
+		public:
+			FAudioAnalyzeTask(
+					TWeakObjectPtr<UAudioAnalyzerNRT> InAnalyzerUObject, 
+					TUniquePtr<Audio::FAnalyzerNRTBatch>&& InAnalyzerFacade, 
+					TArray<uint8>&& InRawWaveData,
+					int32 InNumChannels,
+					float InSampleRate)
+			: AnalyzerUObject(InAnalyzerUObject)
+			, AnalyzerFacade(MoveTemp(InAnalyzerFacade))
+			, RawWaveData(MoveTemp(InRawWaveData))
+			, NumChannels(InNumChannels)
+			, SampleRate(InSampleRate)
+			{}
+
+			void DoWork()
+			{
+				TUniquePtr<Audio::IAnalyzerNRTResult> Result = AnalyzerFacade->AnalyzePCM16Audio(RawWaveData, NumChannels, SampleRate);
+
+				// Make sharedptr so can be passed by value to async lambda
+				TSharedPtr<Audio::IAnalyzerNRTResult, ESPMode::ThreadSafe> ResultPtr(Result.Release());
+				// Make local so can be passed to lambda without having to capture 'this'
+				TWeakObjectPtr<UAudioAnalyzerNRT> Analyzer = AnalyzerUObject;
+
+				// Set value on game thread.
+				AsyncTask(ENamedThreads::GameThread, [Analyzer, ResultPtr]() {
+					if (Analyzer.IsValid())
+					{
+						Analyzer->SetResult(ResultPtr);
+					}
+				});
+			}
+
+			FORCEINLINE TStatId GetStatId() const { RETURN_QUICK_DECLARE_CYCLE_STAT(AudioAnalyzeTask, STATGROUP_ThreadPoolAsyncTasks); }
+
+		private:
+			TWeakObjectPtr<UAudioAnalyzerNRT> AnalyzerUObject;
+			TUniquePtr<Audio::FAnalyzerNRTBatch> AnalyzerFacade;
+			TArray<uint8> RawWaveData;
+			int32 NumChannels;
+			float SampleRate;
+	};
+}
+
+/*****************************************************/
+/*********** UAudioAnalyzerNRTSettings ***************/
+/*****************************************************/
+
+void UAudioAnalyzerNRTSettings::PostEditChangeProperty (struct FPropertyChangedEvent & PropertyChangedEvent) 
+{
+	Super::PostEditChangeProperty(PropertyChangedEvent);
+
+	if (ShouldEventTriggerAnalysis(PropertyChangedEvent))
+	{
+		AnalyzeAudioDelegate.ExecuteIfBound();
+	}
+}
+
+bool UAudioAnalyzerNRTSettings::ShouldEventTriggerAnalysis(struct FPropertyChangedEvent & PropertyChangeEvent)
+{
+	// By default, all changes to settings will trigger analysis.
+	return true;
+}
+
+
+/*****************************************************/
+/***********      UAudioAnalyzerNRT    ***************/
+/*****************************************************/
+
+void UAudioAnalyzerNRT::PreEditChange(UProperty* PropertyAboutToChange)
+{
+	// If the settings object is replaced, need to unbind any existing settings objects
+	// from calling the analyze audio delegate.
+	Super::PreEditChange(PropertyAboutToChange);
+	
+	UAudioAnalyzerNRTSettings* Settings = GetSettingsFromProperty(PropertyAboutToChange);
+
+	if (Settings)
+	{
+		Settings->AnalyzeAudioDelegate.Unbind();
+	}
+}
+
+void UAudioAnalyzerNRT::PostEditChangeProperty (struct FPropertyChangedEvent & PropertyChangedEvent) 
+{
+	Super::PostEditChangeProperty(PropertyChangedEvent);
+	
+	// Check if the edited property was a UAudioAnalyzerNRTSettings object
+	UAudioAnalyzerNRTSettings* Settings = GetSettingsFromProperty(PropertyChangedEvent.Property);
+
+	if (Settings)
+	{
+		// If it was a UAudioAnalyzerNRTSettings object, bind the FAnalyzeAudioDelegate
+		Settings->AnalyzeAudioDelegate.BindUObject(this, &UAudioAnalyzerNRT::AnalyzeAudio);
+	}
+
+	if (ShouldEventTriggerAnalysis(PropertyChangedEvent))
+	{
+		AnalyzeAudio();
+	}
+}
+
+bool UAudioAnalyzerNRT::ShouldEventTriggerAnalysis(struct FPropertyChangedEvent & PropertyChangeEvent)
+{
+	// by default, all changes will trigger analysis
+	return true;
+}
+
+void UAudioAnalyzerNRT::AnalyzeAudio()
+{
+	TSharedPtr<Audio::IAnalyzerNRTResult, ESPMode::ThreadSafe> NewResult;
+
+	if (nullptr != Sound)
+	{
+		// Read audio while Sound object is assured safe. 
+		if (Sound->ChannelSizes.Num() > 0)
+		{
+			UE_LOG(LogAudioAnalyzer, Warning, TEXT("Soundwave '%s' has multi-channel audio (channels greater than 2). Audio analysis is not currently supported for this yet."), *Sound->GetFullName());
+			return;
+		}
+
+		// Retrieve the raw imported data
+		TArray<uint8> RawWaveData;
+		uint32 SampleRate = 0;
+		uint16 NumChannels = 0;
+
+		if (!Sound->GetImportedSoundWaveData(RawWaveData, SampleRate, NumChannels))
+		{
+			UE_LOG(LogAudioAnalyzer, Error, TEXT("Could not analyze audio due to failed import of sound wave data from Soundwave '%s'."), *Sound->GetFullName());
+			return;
+		}
+
+		if (SampleRate == 0 || NumChannels == 0)
+		{
+			UE_LOG(LogAudioAnalyzer, Error, TEXT("Failed to parse the raw imported data for '%s' for analysis."), *Sound->GetFullName());
+			return;
+		}
+		
+		// Create analyzer helper object
+		TUniquePtr<Audio::FAnalyzerNRTBatch> BatchAnalyzer = MakeUnique<Audio::FAnalyzerNRTBatch>(GetSettings(), GetAnalyzerNRTFactoryName());
+
+		// Use weak reference in case this object is deleted before analysis is done
+		TWeakObjectPtr<UAudioAnalyzerNRT> AnalyzerPtr(this);
+		
+		// Create async task
+		FAutoDeleteAsyncTask<FAudioAnalyzeTask>* AsyncTask = new FAutoDeleteAsyncTask<FAudioAnalyzeTask>(
+					AnalyzerPtr,
+					MoveTemp(BatchAnalyzer), 
+					MoveTemp(RawWaveData),
+					NumChannels,
+					SampleRate);
+
+		// Start task
+		AsyncTask->StartBackgroundTask();
+	}
+	else
+	{
+		// Copy empty result to this object
+		SetResult(nullptr);
+	}
+}
+
+// Returns UAudioAnalyzerNRTSettings* if property points to a valid UAudioAnalyzerNRTSettings, otherwise returns nullptr.
+UAudioAnalyzerNRTSettings* UAudioAnalyzerNRT::GetSettingsFromProperty(UProperty* Property)
+{
+	if (nullptr == Property)
+	{
+		return nullptr;
+	}
+
+	if (Property->IsA(UObjectPropertyBase::StaticClass()))
+	{
+		UObjectPropertyBase* ObjectPropertyBase = CastChecked<UObjectPropertyBase>(Property);
+		
+		if (nullptr == ObjectPropertyBase)
+		{
+			return nullptr;
+		}
+
+		if (ObjectPropertyBase->PropertyClass->IsChildOf(UAudioAnalyzerNRTSettings::StaticClass()))
+		{
+			UObject* PropertyObject = ObjectPropertyBase->GetObjectPropertyValue_InContainer(this);
+			return Cast<UAudioAnalyzerNRTSettings>(PropertyObject);
+		}
+	}
+
+	return nullptr;
+}
+
+void UAudioAnalyzerNRT::SetResult(TSharedPtr<Audio::IAnalyzerNRTResult, ESPMode::ThreadSafe> NewResult)
+{
+	FScopeLock ResultLock(&ResultCriticalSection);
+	Result = NewResult;
+}
+
+#endif
+
+
+void UAudioAnalyzerNRT::Serialize(FArchive& Ar)
+{
+	// default uobject serialize
+	Super::Serialize(Ar);
+
+	// When loading object, Result pointer is invalid. Need to create a valid 
+	// result object for loading.
+	if (!Result.IsValid())
+	{
+		if (!GetClass()->HasAnyClassFlags(CLASS_Abstract))
+		{
+			Audio::IAnalyzerNRTFactory* Factory = Audio::GetAnalyzerNRTFactory(GetAnalyzerNRTFactoryName());
+
+			if (nullptr != Factory)
+			{
+				// Create result and worker from factory
+				TUniquePtr<Audio::IAnalyzerNRTResult> NewResult = Factory->NewResult();
+				{
+					FScopeLock ResultLock(&ResultCriticalSection);
+					Result = TSharedPtr<Audio::IAnalyzerNRTResult, ESPMode::ThreadSafe>(NewResult.Release());
+				}
+			}
+		}
+	}
+
+	if (Result.IsValid())
+	{
+		FScopeLock ResultLock(&ResultCriticalSection);
+		Result->Serialize(Ar);
+	}
+}
+
+TUniquePtr<Audio::IAnalyzerNRTSettings> UAudioAnalyzerNRT::GetSettings()
+{
+	return MakeUnique<Audio::IAnalyzerNRTSettings>();
+}
+
diff --git a/Engine/Source/Runtime/AudioAnalyzer/Private/AudioAnalyzerNRTFacade.cpp b/Engine/Source/Runtime/AudioAnalyzer/Private/AudioAnalyzerNRTFacade.cpp
new file mode 100644
--- /dev/null
+++ b/Engine/Source/Runtime/AudioAnalyzer/Private/AudioAnalyzerNRTFacade.cpp
@@ -0,0 +1,92 @@
+// Copyright 1998-2019 Epic Games, Inc. All Rights Reserved.
+
+#include "AudioAnalyzerNRTFacade.h"
+#include "IAudioAnalyzerNRTInterface.h"
+
+namespace Audio
+{
+	IAnalyzerNRTFactory* GetAnalyzerNRTFactory(FName InFactoryName)
+	{
+		// Get all analyzer nrt factories implementations.
+		TArray<IAnalyzerNRTFactory*> RegisteredFactories = IModularFeatures::Get().GetModularFeatureImplementations<IAnalyzerNRTFactory>(IAnalyzerNRTFactory::GetModularFeatureName());
+
+		// Get the factory of interest by matching the name.
+		TArray<Audio::IAnalyzerNRTFactory*> MatchingFactories = RegisteredFactories.FilterByPredicate([InFactoryName](IAnalyzerNRTFactory* Factory) { check(nullptr != Factory); return Factory->GetName() == InFactoryName; });
+
+		if (0 == MatchingFactories.Num())
+		{
+			// There is a likely programming error if the factory is not found. 
+			UE_LOG(LogAudioAnalyzer, Error, TEXT("Failed to find factory of type '%s' with name '%s'"), *IAnalyzerNRTFactory::GetModularFeatureName().ToString(), *InFactoryName.ToString());
+
+			return nullptr;
+		}
+
+		if (MatchingFactories.Num() > 1)
+		{
+			// Like the Highlander, there should be only one. If multiple factories with the same name exist, the first one in the array will be used. 
+			UE_LOG(LogAudioAnalyzer, Warning, TEXT("Found multiple factories of type '%s' with name '%s'. Factory names should be unique."), *IAnalyzerNRTFactory::GetModularFeatureName().ToString(), *InFactoryName.ToString());
+		}
+
+		return MatchingFactories[0];
+	}
+
+	/**********************************************************/
+	/*************** FAnalyzerNRTBatch ************************/
+	/**********************************************************/
+
+	/**
+	 * Create an FAnalyzerNRTBatch with the analyzer settings and factory name.
+	 */
+	FAnalyzerNRTBatch::FAnalyzerNRTBatch(TUniquePtr<IAnalyzerNRTSettings> InSettings, const FName& InFactoryName)
+	: Settings(MoveTemp(InSettings))
+	, FactoryName(InFactoryName)
+	{}
+
+	/**
+	 * Analyze an entire PCM16 encoded audio object.  Audio for the entire sound should be contained within InRawWaveData.
+	 */
+	TUniquePtr<IAnalyzerNRTResult> FAnalyzerNRTBatch::AnalyzePCM16Audio(const TArray<uint8>& InRawWaveData, int32 InNumChannels, float InSampleRate)
+	{
+		// Get analyzer factory
+		IAnalyzerNRTFactory* Factory = GetAnalyzerNRTFactory(FactoryName);
+
+		if (nullptr == Factory)
+		{
+			UE_LOG(LogAudioAnalyzer, Error, TEXT("Cannot analyze audio due to null factory"));
+
+			return TUniquePtr<IAnalyzerNRTResult>();
+		}
+
+		// Create result and worker from fractory
+		FAnalyzerNRTParameters AnalyzerNRTParameters(InSampleRate, InNumChannels);
+
+		TUniquePtr<IAnalyzerNRTResult> Result = Factory->NewResult();
+		TUniquePtr<IAnalyzerNRTWorker> Worker = Factory->NewWorker(AnalyzerNRTParameters, Settings.Get());
+
+		// Check that worker created successfully
+		if (!Worker.IsValid())
+		{
+			UE_LOG(LogAudioAnalyzer, Error, TEXT("Failed to create IAnalyzerNRTWorker with factory of type '%s' with name '%s'"), *IAnalyzerNRTFactory::GetModularFeatureName().ToString(), *Factory->GetName().ToString());
+
+			return TUniquePtr<IAnalyzerNRTResult>();
+		}
+
+		// Check that result created successfully
+		if (!Result.IsValid())
+		{
+			UE_LOG(LogAudioAnalyzer, Error, TEXT("Failed to create IAnalyzerNRTResult with factory of type '%s' with name '%s'"), *Audio::IAnalyzerNRTFactory::GetModularFeatureName().ToString(), *Factory->GetName().ToString());
+
+			return TUniquePtr<IAnalyzerNRTResult>();
+		}
+
+		// Convert 16 bit pcm to 32 bit float
+		TSampleBuffer<float> FloatSamples((const int16*)InRawWaveData.GetData(), InRawWaveData.Num(), InNumChannels, InSampleRate);
+
+		// Perform and finalize audio analysis.
+		Worker->Analyze(FloatSamples, Result.Get());
+		Worker->Finalize(Result.Get());
+
+		return Result;
+	}
+}
+
diff --git a/Engine/Source/Runtime/AudioAnalyzer/Private/IAudioAnalyzerNRTInterface.cpp b/Engine/Source/Runtime/AudioAnalyzer/Private/IAudioAnalyzerNRTInterface.cpp
new file mode 100644
--- /dev/null
+++ b/Engine/Source/Runtime/AudioAnalyzer/Private/IAudioAnalyzerNRTInterface.cpp
@@ -0,0 +1,38 @@
+// Copyright 1998-2019 Epic Games, Inc. All Rights Reserved.
+
+#include "IAudioAnalyzerNRTInterface.h"
+#include "CoreMinimal.h"
+
+namespace Audio
+{
+	void IAnalyzerNRTResult::CopyFrom(IAnalyzerNRTResult* SourceResult)
+	{
+		FBitWriter SerializedResult;
+		SourceResult->Serialize(SerializedResult);
+		FBitReader DeserializedResult = FBitReader(SerializedResult.GetData(), SerializedResult.GetNumBytes() * 8);
+		Serialize(DeserializedResult);
+	}
+	
+	FString IAnalyzerNRTResult::ToString()
+	{
+		return FString(TEXT("To use, override ToString() in this implementation of IAnalyzerNRTResult."));
+	}
+
+	FName IAnalyzerNRTFactory::GetModularFeatureName()
+	{
+		static FName AudioExtFeatureName = FName(TEXT("AudioAnalyzerNRTPlugin"));
+		return AudioExtFeatureName;
+	}
+
+	// This is used to identify analyzers created with this factory.
+	FName IAnalyzerNRTFactory::GetName() const
+	{
+		static FName DefaultName(TEXT("UnnamedAudioAnalyzerNRT"));
+		return DefaultName;
+	}
+
+	FString IAnalyzerNRTFactory::GetTitle() const
+	{
+		return FString("Unnamed Non Real-Time Audio Analyzer.");
+	}
+}
diff --git a/Engine/Source/Runtime/AudioAnalyzer/Public/AudioAnalyzerModule.h b/Engine/Source/Runtime/AudioAnalyzer/Public/AudioAnalyzerModule.h
new file mode 100644
--- /dev/null
+++ b/Engine/Source/Runtime/AudioAnalyzer/Public/AudioAnalyzerModule.h
@@ -0,0 +1,15 @@
+// Copyright 1998-2019 Epic Games, Inc. All Rights Reserved.
+
+#pragma once
+
+#include "CoreMinimal.h"
+
+DECLARE_LOG_CATEGORY_EXTERN(LogAudioAnalyzer, Log, All);
+
+class FAudioAnalyzerModule : public IModuleInterface
+{
+	public:
+	/** IModuleInterface implementation */
+	void StartupModule();
+	void ShutdownModule();
+};
diff --git a/Engine/Source/Runtime/AudioAnalyzer/Public/AudioAnalyzerNRTFacade.h b/Engine/Source/Runtime/AudioAnalyzer/Public/AudioAnalyzerNRTFacade.h
new file mode 100644
--- /dev/null
+++ b/Engine/Source/Runtime/AudioAnalyzer/Public/AudioAnalyzerNRTFacade.h
@@ -0,0 +1,36 @@
+// Copyright 1998-2019 Epic Games, Inc. All Rights Reserved.
+
+#pragma once
+
+#include "CoreMinimal.h"
+#include "IAudioAnalyzerNRTInterface.h"
+
+namespace Audio
+{
+
+	IAnalyzerNRTFactory* GetAnalyzerNRTFactory(FName InFactoryName);
+
+	/** FAnalyzerNRTBatch
+	 *
+	 * FAnalyzerNRTBatch provides a simplified interface for running 
+	 * analyzer factories over complete audio resources.
+	 */
+	class AUDIOANALYZER_API FAnalyzerNRTBatch
+	{
+		public:
+			/**
+			 * Create an FAnalyzerNRTBatch with the analyzer settings and factory name.
+			 */
+			FAnalyzerNRTBatch(TUniquePtr<IAnalyzerNRTSettings> InSettings, const FName& InFactoryName);
+
+			/**
+			 * Analyze an entire PCM16 encoded audio object.  Audio for the entire sound should be contained within InRawWaveData.
+			 */
+			TUniquePtr<IAnalyzerNRTResult> AnalyzePCM16Audio(const TArray<uint8>& InRawWaveData, int32 InNumChannels, float InSampleRate);
+
+		private:
+
+			TUniquePtr<IAnalyzerNRTSettings> Settings;
+			FName FactoryName;
+	};
+}
diff --git a/Engine/Source/Runtime/AudioAnalyzer/Public/IAudioAnalyzerNRTInterface.h b/Engine/Source/Runtime/AudioAnalyzer/Public/IAudioAnalyzerNRTInterface.h
new file mode 100644
--- /dev/null
+++ b/Engine/Source/Runtime/AudioAnalyzer/Public/IAudioAnalyzerNRTInterface.h
@@ -0,0 +1,112 @@
+// Copyright 1998-2019 Epic Games, Inc. All Rights Reserved.
+
+/** Plugin interface for Non Real-Time (NRT) Audio Analyzers. */
+
+#pragma once
+
+#include "CoreMinimal.h"
+#include "Sound/SampleBuffer.h"
+#include "Features/IModularFeature.h"
+
+
+namespace Audio
+{
+	
+	/** FAnalyzerNRTParameters
+	 *
+	 * These parameters are pass to an IAnalyzerNRTFactory when creating
+	 * a new IAnalyzerNRTWorker or IAnalyzerNRTResult
+	 */
+	struct AUDIOANALYZER_API FAnalyzerNRTParameters
+	{
+		public:
+			float SampleRate;
+			int32 NumChannels;
+
+		FAnalyzerNRTParameters(float InSampleRate=0.0f, int32 InNumChannels=0)
+		:	SampleRate(InSampleRate)
+		,	NumChannels(InNumChannels)
+		{}
+	};
+
+	/** IAnalyzerNRTSettings
+	 *
+	 * This interface defines the required methods for non-real-time
+	 * analyzer settings.
+	 */
+	class AUDIOANALYZER_API IAnalyzerNRTSettings
+	{
+		public:
+			virtual ~IAnalyzerNRTSettings() {};
+	};
+
+	/** IAnalyzerNRTResult
+	 *
+	 * This interface defines the required methods for non-real-time
+	 * analyzer results.
+	 */
+	class AUDIOANALYZER_API IAnalyzerNRTResult
+	{
+	public:
+		virtual ~IAnalyzerNRTResult() {};
+
+		// This is used to define how to serialize this instance of results for remote profiling.
+		virtual void Serialize(FArchive& Archive) = 0;
+
+		// This virtual can be overridden to provide a faster copying scheme than full serialization
+		// when analyzing non-remote targets. If not overridden, this function will use the Serialize call.
+		virtual void CopyFrom(IAnalyzerNRTResult* SourceResult);
+
+		// This can be overridden to return a string description of this results struct.
+		virtual FString ToString();
+	};
+
+	/** IAnalyzerNRTWorker
+	 *
+	 * This interface is used to define a class that will handle actual
+	 * analysis of a singular audio asset.
+	 */
+	class AUDIOANALYZER_API IAnalyzerNRTWorker
+	{
+	public:
+		virtual ~IAnalyzerNRTWorker() {};
+
+		/** 
+		 * Perform analysis of an audio stream.
+		 * This method may be called multiple times with audio from the same source.
+		 */
+		virtual void Analyze(const TSampleBuffer<float>& InAudio, IAnalyzerNRTResult* OutResult) = 0;
+
+		/** 
+		 * Called when analysis of audio asset is complete. 
+		 */
+		virtual void Finalize(IAnalyzerNRTResult* OutResult) = 0;
+	};
+
+	/** IAnalyzerNRTFactory
+	 *
+	 * This is used to define a non real-time analyzer.
+	 */
+	class AUDIOANALYZER_API IAnalyzerNRTFactory : public IModularFeature
+	{
+	public:
+		virtual ~IAnalyzerNRTFactory() {};
+
+		// Supplied unique name of IAnalyzerNRTFactory to enable querying of added 
+		// analyzer factories
+		static FName GetModularFeatureName();
+
+		// Name of specific analyzer type.
+		virtual FName GetName() const;
+
+		// Human readable name of analyzer.
+		virtual FString GetTitle() const;
+
+		// Create a new result.
+		virtual TUniquePtr<IAnalyzerNRTResult> NewResult() = 0;
+
+		// Create a new worker.
+		virtual TUniquePtr<IAnalyzerNRTWorker> NewWorker(const FAnalyzerNRTParameters& InParams, const IAnalyzerNRTSettings* InSettings) = 0;
+	};
+}
+
