Fix CIS warnings.

#rnx

diff --git a/Engine/Source/Runtime/AudioAnalyzer/AudioAnalyzer.Build.cs b/Engine/Source/Runtime/AudioAnalyzer/AudioAnalyzer.Build.cs
--- a/Engine/Source/Runtime/AudioAnalyzer/AudioAnalyzer.Build.cs
+++ b/Engine/Source/Runtime/AudioAnalyzer/AudioAnalyzer.Build.cs
@@ -6,7 +6,7 @@ namespace UnrealBuildTool.Rules
 	{
 		public AudioAnalyzer(ReadOnlyTargetRules Target) : base(Target)
 		{
-			bFasterWithoutUnity = true;
+			bUseUnity = true;
 
             PrivateDependencyModuleNames.AddRange(
 				new string[]
