Removing memory leak warning for autodelete task

diff --git a/Engine/Source/Runtime/AudioAnalyzer/Private/AudioAnalyzerNRT.cpp b/Engine/Source/Runtime/AudioAnalyzer/Private/AudioAnalyzerNRT.cpp
--- a/Engine/Source/Runtime/AudioAnalyzer/Private/AudioAnalyzerNRT.cpp
+++ b/Engine/Source/Runtime/AudioAnalyzer/Private/AudioAnalyzerNRT.cpp
@@ -157,16 +157,8 @@ void UAudioAnalyzerNRT::AnalyzeAudio()
 		// Use weak reference in case this object is deleted before analysis is done
 		TWeakObjectPtr<UAudioAnalyzerNRT> AnalyzerPtr(this);
 		
-		// Create async task
-		FAutoDeleteAsyncTask<FAudioAnalyzeTask>* AsyncTask = new FAutoDeleteAsyncTask<FAudioAnalyzeTask>(
-					AnalyzerPtr,
-					MoveTemp(BatchAnalyzer), 
-					MoveTemp(RawWaveData),
-					NumChannels,
-					SampleRate);
-
-		// Start task
-		AsyncTask->StartBackgroundTask();
+		// Create and start async task. Parentheses avoids memory leak warnings from static analysis.
+		(new FAutoDeleteAsyncTask<FAudioAnalyzeTask>(AnalyzerPtr, MoveTemp(BatchAnalyzer), MoveTemp(RawWaveData), NumChannels, SampleRate))->StartBackgroundTask();
 	}
 	else
 	{
