moving include for SampleBuffer.h?

diff --git a/Engine/Source/Runtime/AudioAnalyzer/Public/IAudioAnalyzerNRTInterface.h b/Engine/Source/Runtime/AudioAnalyzer/Public/IAudioAnalyzerNRTInterface.h
--- a/Engine/Source/Runtime/AudioAnalyzer/Public/IAudioAnalyzerNRTInterface.h
+++ b/Engine/Source/Runtime/AudioAnalyzer/Public/IAudioAnalyzerNRTInterface.h
@@ -5,7 +5,7 @@
 #pragma once
 
 #include "CoreMinimal.h"
-#include "Sound/SampleBuffer.h"
+#include "SampleBuffer.h"
 #include "Features/IModularFeature.h"
 
 namespace Audio
