Adding ability for audio analyzer to have native C++ delegate callbacks on analysis results.

Allowing creating analyzer without a world context object (and a world instead)

diff --git a/Engine/Plugins/Runtime/AudioSynesthesia/Source/AudioSynesthesia/Classes/Meter.h b/Engine/Plugins/Runtime/AudioSynesthesia/Source/AudioSynesthesia/Classes/Meter.h
--- a/Engine/Plugins/Runtime/AudioSynesthesia/Source/AudioSynesthesia/Classes/Meter.h
+++ b/Engine/Plugins/Runtime/AudioSynesthesia/Source/AudioSynesthesia/Classes/Meter.h
@@ -90,15 +90,27 @@ struct AUDIOSYNESTHESIA_API FMeterResults
 /** Delegate to receive all overall loudness results (time-stamped in an array) since last delegate call. */
 DECLARE_DYNAMIC_MULTICAST_DELEGATE_OneParam(FOnOverallMeterResults, const TArray<FMeterResults>&, MeterResults);
 
+/** shadow delegate declaration for above */
+DECLARE_MULTICAST_DELEGATE_TwoParams(FOnOverallMeterResultsNative, UMeterAnalyzer*, const TArray<FMeterResults>&);
+
 /** Delegate to receive only the most recent overall meter results. */
 DECLARE_DYNAMIC_MULTICAST_DELEGATE_OneParam(FOnLatestOverallMeterResults, const FMeterResults&, LatestOverallMeterResults);
 
+/** shadow delegate declaration for above */
+DECLARE_MULTICAST_DELEGATE_TwoParams(FOnLatestOverallMeterResultsNative, UMeterAnalyzer*, const FMeterResults&);
+
 /** Delegate to receive all meter results per channel (time-stamped in an array) since last delegate call. */
 DECLARE_DYNAMIC_MULTICAST_DELEGATE_TwoParams(FOnPerChannelMeterResults, int32, ChannelIndex, const TArray<FMeterResults>&, MeterResults);
 
+/** shadow delegate declaration for above */
+DECLARE_MULTICAST_DELEGATE_ThreeParams(FOnPerChannelMeterResultsNative, UMeterAnalyzer*, int32, const TArray<FMeterResults>&);
+
 /** Delegate to receive only the most recent overall meter result per channel. */
 DECLARE_DYNAMIC_MULTICAST_DELEGATE_TwoParams(FOnLatestPerChannelMeterResults, int32, ChannelIndex, const FMeterResults&, LatestMeterResults);
 
+/** shadow delegate declaration for above */
+DECLARE_MULTICAST_DELEGATE_ThreeParams(FOnLatestPerChannelMeterResultsNative, UMeterAnalyzer*, int32, const FMeterResults&);
+
 /** UMeterAnalyzer
  *
  * UMeterAnalyzer calculates the current amplitude of an
@@ -121,18 +133,26 @@ public:
 	UPROPERTY(BlueprintAssignable)
 	FOnOverallMeterResults OnOverallMeterResults;
 
+	FOnOverallMeterResultsNative OnOverallMeterResultsNative;
+
 	/** Delegate to receive all meter results, per-channel, since last delegate call. */
 	UPROPERTY(BlueprintAssignable)
 	FOnPerChannelMeterResults OnPerChannelMeterResults;
 
+	FOnPerChannelMeterResultsNative OnPerChannelMeterResultsNative;
+
 	/** Delegate to receive the latest overall meter results. */
 	UPROPERTY(BlueprintAssignable)
 	FOnLatestOverallMeterResults OnLatestOverallMeterResults;
 
+	FOnLatestOverallMeterResultsNative OnLatestOverallMeterResultsNative;
+
 	/** Delegate to receive the latest per-channel meter results. */
 	UPROPERTY(BlueprintAssignable)
 	FOnLatestPerChannelMeterResults OnLatestPerChannelMeterResults;
 
+	FOnLatestPerChannelMeterResultsNative OnLatestPerChannelMeterResultsNative;
+
 	/** Convert UMeterSettings to IAnalyzerSettings */
 	TUniquePtr<Audio::IAnalyzerSettings> GetSettings(const int32 InSampleRate, const int32 InNumChannels) const override;
 
diff --git a/Engine/Plugins/Runtime/AudioSynesthesia/Source/AudioSynesthesia/Private/Meter.cpp b/Engine/Plugins/Runtime/AudioSynesthesia/Source/AudioSynesthesia/Private/Meter.cpp
--- a/Engine/Plugins/Runtime/AudioSynesthesia/Source/AudioSynesthesia/Private/Meter.cpp
+++ b/Engine/Plugins/Runtime/AudioSynesthesia/Source/AudioSynesthesia/Private/Meter.cpp
@@ -104,8 +104,8 @@ void UMeterAnalyzer::BroadcastResults()
 
 	if (NumChannels > 0)
 	{
-		bool bIsOnOverallMeterResultsBound = OnOverallMeterResults.IsBound();
-		bool bIsOnLatestOverallMeterResultsBound = OnLatestOverallMeterResults.IsBound();
+		bool bIsOnOverallMeterResultsBound = OnOverallMeterResults.IsBound() || OnOverallMeterResultsNative.IsBound();
+		bool bIsOnLatestOverallMeterResultsBound = OnLatestOverallMeterResults.IsBound() || OnLatestOverallMeterResults.IsBound();
 		if (bIsOnOverallMeterResultsBound || bIsOnLatestOverallMeterResultsBound)
 		{
 			const TArray<Audio::FMeterEntry>& OverallMeterArray = MeterResults->GetMeterArray();
@@ -114,20 +114,18 @@ void UMeterAnalyzer::BroadcastResults()
 				TArray<FMeterResults> Results = ConvertToBlueprintResults(Settings, OverallMeterArray);
 				check(Results.Num() > 0);
 
-				if (bIsOnOverallMeterResultsBound)
-				{
-					OnOverallMeterResults.Broadcast(Results);
-				}
-				if (bIsOnLatestOverallMeterResultsBound)
-				{
-					FMeterResults& Latest = Results[Results.Num() - 1];
-					OnLatestOverallMeterResults.Broadcast(Latest);
-				}
+				OnOverallMeterResults.Broadcast(Results);
+				OnOverallMeterResultsNative.Broadcast(this, Results);
+
+				FMeterResults& Latest = Results[Results.Num() - 1];
+				OnLatestOverallMeterResults.Broadcast(Latest);
+				OnLatestOverallMeterResultsNative.Broadcast(this, Latest);
+
 			}
 		}
 
-		bool bIsOnPerChannelMeterResultsBound = OnPerChannelMeterResults.IsBound();
-		bool bIsOnLatestPerChannelMeterResultsBound = OnLatestPerChannelMeterResults.IsBound();
+		bool bIsOnPerChannelMeterResultsBound = OnPerChannelMeterResults.IsBound() || OnPerChannelMeterResultsNative.IsBound();
+		bool bIsOnLatestPerChannelMeterResultsBound = OnLatestPerChannelMeterResults.IsBound() || OnLatestPerChannelMeterResultsNative.IsBound();
 
 		if (bIsOnPerChannelMeterResultsBound || bIsOnLatestPerChannelMeterResultsBound)
 		{
@@ -139,15 +137,11 @@ void UMeterAnalyzer::BroadcastResults()
 					TArray<FMeterResults> Results = ConvertToBlueprintResults(Settings, MeterArray);
 					check(Results.Num() > 0);
 
-					if (bIsOnPerChannelMeterResultsBound)
-					{
-						OnPerChannelMeterResults.Broadcast(ChannelIndex, Results);
-					}
+					OnPerChannelMeterResults.Broadcast(ChannelIndex, Results);
+					OnPerChannelMeterResultsNative.Broadcast(this, ChannelIndex, Results);
 
-					if (bIsOnLatestPerChannelMeterResultsBound)
-					{
-						OnLatestPerChannelMeterResults.Broadcast(ChannelIndex, Results[Results.Num() - 1]);
-					}
+					OnLatestPerChannelMeterResults.Broadcast(ChannelIndex, Results[Results.Num() - 1]);
+					OnLatestPerChannelMeterResultsNative.Broadcast(this, ChannelIndex, Results[Results.Num() - 1]);
 				}
 			}
 		}
diff --git a/Engine/Source/Runtime/AudioAnalyzer/Classes/AudioAnalyzer.h b/Engine/Source/Runtime/AudioAnalyzer/Classes/AudioAnalyzer.h
--- a/Engine/Source/Runtime/AudioAnalyzer/Classes/AudioAnalyzer.h
+++ b/Engine/Source/Runtime/AudioAnalyzer/Classes/AudioAnalyzer.h
@@ -88,6 +88,9 @@ public:
 	UFUNCTION(BlueprintCallable, BlueprintCosmetic, Category = AudioAnalyzer, meta = (WorldContext = "WorldContextObject"))
 	void StartAnalyzing(const UObject* WorldContextObject, UAudioBus* AudioBusToAnalyze);
 
+	/** Starts analyzing using the given world.*/
+	void StartAnalyzing(UWorld* InWorld, UAudioBus* AudioBusToAnalyze);
+
 	/** Stops analyzing audio. */
 	UFUNCTION(BlueprintCallable, BlueprintCosmetic, Category = AudioAnalyzer, meta = (WorldContext = "WorldContextObject"))
 	void StopAnalyzing(const UObject* WorldContextObject);
diff --git a/Engine/Source/Runtime/AudioAnalyzer/Classes/AudioAnalyzerSubsystem.h b/Engine/Source/Runtime/AudioAnalyzer/Classes/AudioAnalyzerSubsystem.h
--- a/Engine/Source/Runtime/AudioAnalyzer/Classes/AudioAnalyzerSubsystem.h
+++ b/Engine/Source/Runtime/AudioAnalyzer/Classes/AudioAnalyzerSubsystem.h
@@ -25,6 +25,7 @@ public:
 	//~ Begin FTickableGameObject Interface
 	virtual void Tick(float DeltaTime) override;
 	virtual bool IsTickable() const override;
+	virtual bool IsTickableInEditor() const { return true; }
 	virtual TStatId GetStatId() const override;
 	//~ End FTickableGameObject Interface
 
diff --git a/Engine/Source/Runtime/AudioAnalyzer/Private/AudioAnalyzer.cpp b/Engine/Source/Runtime/AudioAnalyzer/Private/AudioAnalyzer.cpp
--- a/Engine/Source/Runtime/AudioAnalyzer/Private/AudioAnalyzer.cpp
+++ b/Engine/Source/Runtime/AudioAnalyzer/Private/AudioAnalyzer.cpp
@@ -28,7 +28,7 @@ void FAudioAnalyzeTask::DoWork()
 	Results = AnalyzerFacade->AnalyzeAudioBuffer(AudioData, NumChannels, SampleRate);
 }
 
-void UAudioAnalyzer::StartAnalyzing(const UObject* WorldContextObject, UAudioBus* AudioBusToAnalyze)
+void UAudioAnalyzer::StartAnalyzing(UWorld* InWorld, UAudioBus* AudioBusToAnalyze)
 {
 	if (!AudioBusToAnalyze)
 	{
@@ -36,15 +36,13 @@ void UAudioAnalyzer::StartAnalyzing(const UObject* WorldContextObject, UAudioBus
 		return;
 	}
 
-	// Retrieve the world and if it's not a valid world, don't do anything
-	UWorld* ThisWorld = GEngine->GetWorldFromContextObject(WorldContextObject, EGetWorldErrorMode::LogAndReturnNull);
-	if (!ThisWorld || !ThisWorld->bAllowAudioPlayback || ThisWorld->GetNetMode() == NM_DedicatedServer)
+	if (!InWorld)
 	{
 		return;
 	}
 
-	// Get the audio mixer device to get audio mixer sample rate and be able to play audio buses
-	Audio::FMixerDevice* MixerDevice = FAudioDeviceManager::GetAudioMixerDeviceFromWorldContext(WorldContextObject);
+	Audio::FMixerDevice* MixerDevice = static_cast<Audio::FMixerDevice*>(InWorld->GetAudioDeviceRaw());
+
 	if (!MixerDevice)
 	{
 		UE_LOG(LogAudioAnalyzer, Error, TEXT("Audio analyzer only works with the audio mixer."));
@@ -54,31 +52,31 @@ void UAudioAnalyzer::StartAnalyzing(const UObject* WorldContextObject, UAudioBus
 	// Retrieve the audio analyzer subsystem if it's not already retrieved
 	if (!AudioAnalyzerSubsystem)
 	{
-		AudioAnalyzerSubsystem = UAudioAnalyzerSubsystem::Get(ThisWorld);
+		AudioAnalyzerSubsystem = UAudioAnalyzerSubsystem::Get(InWorld);
 	}
 
 	// Store the audio bus we're analyzing
- 	AudioBus = AudioBusToAnalyze;
+	AudioBus = AudioBusToAnalyze;
+
+	NumBusChannels = AudioBus->GetNumChannels();
+	check(NumBusChannels > 0);
 
- 	NumBusChannels = AudioBus->GetNumChannels();
- 	check(NumBusChannels > 0);
- 	
 	AudioMixerSampleRate = (int32)MixerDevice->GetSampleRate();
- 
- 	// Get the analyzer factory to use for our analysis
+
+	// Get the analyzer factory to use for our analysis
 	if (!AnalyzerFactory)
 	{
- 		AnalyzerFactory = Audio::GetAnalyzerFactory(GetAnalyzerFactoryName());
+		AnalyzerFactory = Audio::GetAnalyzerFactory(GetAnalyzerFactoryName());
 	}
 	checkf(AnalyzerFactory != nullptr, TEXT("Need to register the factory as a modular feature for the analyzer '%s'."), *GetAnalyzerFactoryName().ToString());
 
- 	// Start the audio bus. This won't do anythign if the bus is already started elsewhere.
- 	uint32 AudioBusId = AudioBus->GetUniqueID();
- 	int32 NumChannels = (int32)AudioBus->AudioBusChannels + 1;
- 	MixerDevice->StartAudioBus(AudioBusId, NumChannels, false);
- 
- 	// Get an output patch for the audio bus
- 	PatchOutputStrongPtr = MixerDevice->AddPatchForAudioBus(AudioBusId);
+	// Start the audio bus. This won't do anythign if the bus is already started elsewhere.
+	uint32 AudioBusId = AudioBus->GetUniqueID();
+	int32 NumChannels = (int32)AudioBus->AudioBusChannels + 1;
+	MixerDevice->StartAudioBus(AudioBusId, NumChannels, false);
+
+	// Get an output patch for the audio bus
+	PatchOutputStrongPtr = MixerDevice->AddPatchForAudioBus(AudioBusId);
 
 	// Register this audio analyzer with the audio analyzer subsystem
 	// The subsystem will query this analyzer to see if it has enough audio to perform analysis.
@@ -89,6 +87,25 @@ void UAudioAnalyzer::StartAnalyzing(const UObject* WorldContextObject, UAudioBus
 	AnalyzerFacade = MakeUnique<Audio::FAnalyzerFacade>(GetSettings(AudioMixerSampleRate, NumBusChannels), AnalyzerFactory);
 }
 
+
+void UAudioAnalyzer::StartAnalyzing(const UObject* WorldContextObject, UAudioBus* AudioBusToAnalyze)
+{
+	if (!AudioBusToAnalyze)
+	{
+		UE_LOG(LogAudioAnalyzer, Error, TEXT("Unable to analyze audio without an audio bus to analyze."));
+		return;
+	}
+
+	// Retrieve the world and if it's not a valid world, don't do anything
+	UWorld* ThisWorld = GEngine->GetWorldFromContextObject(WorldContextObject, EGetWorldErrorMode::ReturnNull);
+	if (!ThisWorld || !ThisWorld->bAllowAudioPlayback || ThisWorld->GetNetMode() == NM_DedicatedServer)
+	{
+		return;
+	}
+
+	StartAnalyzing(ThisWorld, AudioBusToAnalyze);
+}
+
 void UAudioAnalyzer::StopAnalyzing(const UObject* WorldContextObject)
 {
 	if (AudioAnalyzerSubsystem)
