Fix for crash when loading a new level and having the metasound editor open.

-- moved the audio analyzer subsystem to be an engine subsystem vs a world subsystem

diff --git a/Engine/Plugins/Runtime/AudioSynesthesia/Source/AudioSynesthesia/Private/AudioSynesthesiaModule.cpp b/Engine/Plugins/Runtime/AudioSynesthesia/Source/AudioSynesthesia/Private/AudioSynesthesiaModule.cpp
--- a/Engine/Plugins/Runtime/AudioSynesthesia/Source/AudioSynesthesia/Private/AudioSynesthesiaModule.cpp
+++ b/Engine/Plugins/Runtime/AudioSynesthesia/Source/AudioSynesthesia/Private/AudioSynesthesiaModule.cpp
@@ -5,6 +5,7 @@
 #include "Modules/ModuleInterface.h"
 #include "Features/IModularFeatures.h"
 #include "LoudnessNRTFactory.h"
+#include "AudioAnalyzerModule.h"
 #include "AudioSynesthesiaLog.h"
 
 DEFINE_LOG_CATEGORY(LogAudioSynesthesia);
@@ -16,6 +17,7 @@ namespace Audio
 	void FAudioSynesthesiaModule::StartupModule()
 	{
 		FModuleManager::Get().LoadModuleChecked<IModuleInterface>("AudioSynesthesiaCore");
+		FModuleManager::Get().LoadModuleChecked<FAudioAnalyzerModule>("AudioAnalyzer");
 	}
 
 	void FAudioSynesthesiaModule::ShutdownModule()
diff --git a/Engine/Source/Runtime/AudioAnalyzer/Classes/AudioAnalyzerSubsystem.h b/Engine/Source/Runtime/AudioAnalyzer/Classes/AudioAnalyzerSubsystem.h
--- a/Engine/Source/Runtime/AudioAnalyzer/Classes/AudioAnalyzerSubsystem.h
+++ b/Engine/Source/Runtime/AudioAnalyzer/Classes/AudioAnalyzerSubsystem.h
@@ -4,7 +4,8 @@
 
 #include "CoreMinimal.h"
 #include "AudioAnalyzer.h"
-#include "Subsystems/WorldSubsystem.h"
+#include "Subsystems/EngineSubsystem.h"
+#include "Tickable.h"
 #include "AudioAnalyzerSubsystem.generated.h"
 
 class UWorld;
@@ -12,8 +13,8 @@ class UWorld;
 /** 
 * Class manages querying analysis results from various audio analyzers.
 */
-UCLASS(DisplayName = "Audio Analysis")
-class AUDIOANALYZER_API UAudioAnalyzerSubsystem : public UTickableWorldSubsystem
+UCLASS()
+class AUDIOANALYZER_API UAudioAnalyzerSubsystem : public UEngineSubsystem, public FTickableGameObject
 {
 	GENERATED_BODY()
 
@@ -24,14 +25,19 @@ public:
 	//~ Begin FTickableGameObject Interface
 	virtual void Tick(float DeltaTime) override;
 	virtual bool IsTickable() const override;
-	virtual bool IsTickableInEditor() const override { return true; }
+	virtual bool IsTickableInEditor() const { return true; }
 	virtual TStatId GetStatId() const override;
 	//~ End FTickableGameObject Interface
 
+	//~ Begin USubsystem
+	virtual void Initialize(FSubsystemCollectionBase& Collection) override;
+	virtual void Deinitialize() override;
+	//~ End USubsystem
+
 	void RegisterAudioAnalyzer(UAudioAnalyzer* InAnalyzer);
 	void UnregisterAudioAnalyzer(UAudioAnalyzer* InAnalyzer);
 
-	static UAudioAnalyzerSubsystem* Get(UWorld* World);
+	static UAudioAnalyzerSubsystem* Get();
 
 private:
 
diff --git a/Engine/Source/Runtime/AudioAnalyzer/Private/AudioAnalyzer.cpp b/Engine/Source/Runtime/AudioAnalyzer/Private/AudioAnalyzer.cpp
--- a/Engine/Source/Runtime/AudioAnalyzer/Private/AudioAnalyzer.cpp
+++ b/Engine/Source/Runtime/AudioAnalyzer/Private/AudioAnalyzer.cpp
@@ -52,7 +52,7 @@ void UAudioAnalyzer::StartAnalyzing(UWorld* InWorld, UAudioBus* AudioBusToAnalyz
 	// Retrieve the audio analyzer subsystem if it's not already retrieved
 	if (!AudioAnalyzerSubsystem)
 	{
-		AudioAnalyzerSubsystem = UAudioAnalyzerSubsystem::Get(InWorld);
+		AudioAnalyzerSubsystem = UAudioAnalyzerSubsystem::Get();
 	}
 
 	// Store the audio bus we're analyzing
@@ -80,8 +80,11 @@ void UAudioAnalyzer::StartAnalyzing(UWorld* InWorld, UAudioBus* AudioBusToAnalyz
 
 	// Register this audio analyzer with the audio analyzer subsystem
 	// The subsystem will query this analyzer to see if it has enough audio to perform analysis.
-	// If it does, it'll report the results (of any previous analysis) and ask us to start analyzing the audio.		
-	AudioAnalyzerSubsystem->RegisterAudioAnalyzer(this);
+	// If it does, it'll report the results (of any previous analysis) and ask us to start analyzing the audio.	
+	if (AudioAnalyzerSubsystem)
+	{
+		AudioAnalyzerSubsystem->RegisterAudioAnalyzer(this);
+	}
 
 	// Setup the analyzer facade here once
 	AnalyzerFacade = MakeUnique<Audio::FAnalyzerFacade>(GetSettings(AudioMixerSampleRate, NumBusChannels), AnalyzerFactory);
diff --git a/Engine/Source/Runtime/AudioAnalyzer/Private/AudioAnalyzerSubsystem.cpp b/Engine/Source/Runtime/AudioAnalyzer/Private/AudioAnalyzerSubsystem.cpp
--- a/Engine/Source/Runtime/AudioAnalyzer/Private/AudioAnalyzerSubsystem.cpp
+++ b/Engine/Source/Runtime/AudioAnalyzer/Private/AudioAnalyzerSubsystem.cpp
@@ -12,11 +12,11 @@ UAudioAnalyzerSubsystem::~UAudioAnalyzerSubsystem()
 	AudioAnalyzers.Reset();
 }
 
-UAudioAnalyzerSubsystem* UAudioAnalyzerSubsystem::Get(UWorld* World)
+UAudioAnalyzerSubsystem* UAudioAnalyzerSubsystem::Get()
 {
-	if (World)
+	if (GEngine)
 	{
-		return World->GetSubsystem<UAudioAnalyzerSubsystem>();
+		return GEngine->GetEngineSubsystem<UAudioAnalyzerSubsystem>();
 	}
 	return nullptr;
 }
@@ -57,6 +57,17 @@ bool UAudioAnalyzerSubsystem::IsTickable() const
 	return false;
 }
 
+void UAudioAnalyzerSubsystem::Initialize(FSubsystemCollectionBase& Collection)
+{
+	Super::Initialize(Collection);
+}
+
+void UAudioAnalyzerSubsystem::Deinitialize()
+{
+	// Release our audio analyzers
+	AudioAnalyzers.Reset();
+}
+
 void UAudioAnalyzerSubsystem::RegisterAudioAnalyzer(UAudioAnalyzer* InAnalyzer)
 {
 	AudioAnalyzers.AddUnique(InAnalyzer);
