Build fix removing unnecessary Super::Tick from CL 15512967

diff --git a/Engine/Source/Runtime/AudioAnalyzer/Private/AudioAnalyzerSubsystem.cpp b/Engine/Source/Runtime/AudioAnalyzer/Private/AudioAnalyzerSubsystem.cpp
--- a/Engine/Source/Runtime/AudioAnalyzer/Private/AudioAnalyzerSubsystem.cpp
+++ b/Engine/Source/Runtime/AudioAnalyzer/Private/AudioAnalyzerSubsystem.cpp
@@ -28,8 +28,6 @@ TStatId UAudioAnalyzerSubsystem::GetStatId() const
 
 void UAudioAnalyzerSubsystem::Tick(float DeltaTime)
 {
-	Super::Tick(DeltaTime);
-
 	// Loop through all analyzers and if they're ready to analyze, do it
 	for (UAudioAnalyzer* Analyzer : AudioAnalyzers)
 	{
