Non-Unity compile fix

diff --git a/Engine/Source/Runtime/AudioAnalyzer/Private/AudioAnalyzerSubsystem.cpp b/Engine/Source/Runtime/AudioAnalyzer/Private/AudioAnalyzerSubsystem.cpp
--- a/Engine/Source/Runtime/AudioAnalyzer/Private/AudioAnalyzerSubsystem.cpp
+++ b/Engine/Source/Runtime/AudioAnalyzer/Private/AudioAnalyzerSubsystem.cpp
@@ -2,6 +2,7 @@
 
 #include "AudioAnalyzerSubsystem.h"
 #include "Engine/World.h"
+#include "Engine/Engine.h"
 
 UAudioAnalyzerSubsystem::UAudioAnalyzerSubsystem()
 {
