[Audio Analyzer] Added IAnalyzerControls + related logic to be able to do live updates.

diff --git a/Engine/Source/Runtime/AudioAnalyzer/Classes/AudioAnalyzer.h b/Engine/Source/Runtime/AudioAnalyzer/Classes/AudioAnalyzer.h
--- a/Engine/Source/Runtime/AudioAnalyzer/Classes/AudioAnalyzer.h
+++ b/Engine/Source/Runtime/AudioAnalyzer/Classes/AudioAnalyzer.h
@@ -32,6 +32,9 @@ class FAudioAnalyzeTask : public FNonAbandonableTask
 public:
 	FAudioAnalyzeTask(TUniquePtr<Audio::FAnalyzerFacade>& InAnalyzerFacade, int32 InSampleRate, int32 InNumChannels);
 	
+	// Set in the task the current state of the analyzer controls
+	void SetAnalyzerControls(TSharedPtr<Audio::IAnalyzerControls> InControls);
+
 	// Give the task the audio data to analyze
 	void SetAudioBuffer(TArray<float>&& InAudioData);
 
@@ -52,6 +55,7 @@ private:
 	int32 NumChannels = 0;
 	TArray<float> AudioData;
 	TUniquePtr<Audio::IAnalyzerResult> Results;
+	TSharedPtr<Audio::IAnalyzerControls> AnalyzerControls;
 };
 
 /** UAudioAnalyzer
@@ -102,6 +106,12 @@ public:
 	 */
 	AUDIOANALYZER_API virtual TUniquePtr<Audio::IAnalyzerSettings> GetSettings(const int32 InSampleRate, const int32 InNumChannels) const;
 
+	/**
+	 * Implementations can override this method to create controls objects
+	 * specific for their analyzer.
+	 */
+	AUDIOANALYZER_API virtual TSharedPtr<Audio::IAnalyzerControls> GetAnalyzerControls() const;
+
 	/** Function to broadcast results. */
 	virtual void BroadcastResults() {}
 
@@ -151,6 +161,9 @@ private:
 	// Cached results of previous analysis task. Can be invalid.
 	TUniquePtr<Audio::IAnalyzerResult> ResultsInternal;
 
+	// The analyzer controls that can be modified real-time
+	TSharedPtr<Audio::IAnalyzerControls> AnalyzerControls;
+
 	// Scratch buffer used for copying data from patch output and fed to analysis task
 	TArray<float> AnalysisBuffer;
 
diff --git a/Engine/Source/Runtime/AudioAnalyzer/Private/AudioAnalyzer.cpp b/Engine/Source/Runtime/AudioAnalyzer/Private/AudioAnalyzer.cpp
--- a/Engine/Source/Runtime/AudioAnalyzer/Private/AudioAnalyzer.cpp
+++ b/Engine/Source/Runtime/AudioAnalyzer/Private/AudioAnalyzer.cpp
@@ -26,10 +26,15 @@ void FAudioAnalyzeTask::SetAudioBuffer(TArray<float>&& InAudioData)
 	AudioData = MoveTemp(InAudioData);
 }
 
+void FAudioAnalyzeTask::SetAnalyzerControls(TSharedPtr<Audio::IAnalyzerControls> InControls)
+{
+	AnalyzerControls = InControls;
+}
+
 void FAudioAnalyzeTask::DoWork()
 {
 	check(AnalyzerFacade);
-	Results = AnalyzerFacade->AnalyzeAudioBuffer(AudioData, NumChannels, SampleRate);
+	Results = AnalyzerFacade->AnalyzeAudioBuffer(AudioData, NumChannels, SampleRate, AnalyzerControls);
 }
 
 void UAudioAnalyzer::StartAnalyzing(UWorld* InWorld, UAudioBus* AudioBusToAnalyze)
@@ -193,6 +198,7 @@ bool UAudioAnalyzer::DoAnalysis()
 	check(NumBusChannels != 0);
 
 	FAudioAnalyzeTask& Task = AnalysisTask->GetTask();
+	Task.SetAnalyzerControls(GetAnalyzerControls());
 	Task.SetAudioBuffer(MoveTemp(AnalysisBuffer));
 
  	AnalysisTask->StartBackgroundTask();
@@ -247,4 +253,7 @@ TUniquePtr<Audio::IAnalyzerSettings> UAudioAnalyzer::GetSettings(const int32 InS
 	return MakeUnique<Audio::IAnalyzerSettings>();
 }
 
-
+TSharedPtr<Audio::IAnalyzerControls> UAudioAnalyzer::GetAnalyzerControls() const
+{
+	return AnalyzerControls;
+}
diff --git a/Engine/Source/Runtime/AudioAnalyzer/Private/AudioAnalyzerFacade.cpp b/Engine/Source/Runtime/AudioAnalyzer/Private/AudioAnalyzerFacade.cpp
--- a/Engine/Source/Runtime/AudioAnalyzer/Private/AudioAnalyzerFacade.cpp
+++ b/Engine/Source/Runtime/AudioAnalyzer/Private/AudioAnalyzerFacade.cpp
@@ -43,7 +43,7 @@ namespace Audio
 	{
 	}
 
-	TUniquePtr<IAnalyzerResult> FAnalyzerFacade::AnalyzeAudioBuffer(const TArray<float>& InAudioBuffer, int32 InNumChannels, float InSampleRate)
+	TUniquePtr<IAnalyzerResult> FAnalyzerFacade::AnalyzeAudioBuffer(const TArray<float>& InAudioBuffer, int32 InNumChannels, float InSampleRate, TSharedPtr<IAnalyzerControls> InControls)
 	{
 		LLM_SCOPE_BYTAG(Audio_Analysis);
 
@@ -80,6 +80,7 @@ namespace Audio
 		}
 
 		// Perform and finalize audio analysis.
+		Worker->SetControls(InControls);
 		Worker->Analyze(MakeArrayView(InAudioBuffer), Result.Get());
 
 		return Result;
diff --git a/Engine/Source/Runtime/AudioAnalyzer/Public/AudioAnalyzerFacade.h b/Engine/Source/Runtime/AudioAnalyzer/Public/AudioAnalyzerFacade.h
--- a/Engine/Source/Runtime/AudioAnalyzer/Public/AudioAnalyzerFacade.h
+++ b/Engine/Source/Runtime/AudioAnalyzer/Public/AudioAnalyzerFacade.h
@@ -29,7 +29,7 @@ namespace Audio
 			/**
 			 * Analyze the audio buffer.
 			 */
-			AUDIOANALYZER_API TUniquePtr<IAnalyzerResult> AnalyzeAudioBuffer(const TArray<float>& InAudioBuffer, int32 InNumChannels, float InSampleRate);
+			AUDIOANALYZER_API TUniquePtr<IAnalyzerResult> AnalyzeAudioBuffer(const TArray<float>& InAudioBuffer, int32 InNumChannels, float InSampleRate, TSharedPtr<IAnalyzerControls> InControls);
 
 		private:
 
diff --git a/Engine/Source/Runtime/AudioAnalyzer/Public/IAudioAnalyzerInterface.h b/Engine/Source/Runtime/AudioAnalyzer/Public/IAudioAnalyzerInterface.h
--- a/Engine/Source/Runtime/AudioAnalyzer/Public/IAudioAnalyzerInterface.h
+++ b/Engine/Source/Runtime/AudioAnalyzer/Public/IAudioAnalyzerInterface.h
@@ -46,6 +46,17 @@ namespace Audio
 		virtual ~IAnalyzerResult() {};
 	};
 
+	/** IAnalyzerControls
+	 *
+	 * This interface defines the controls we want to link to 
+	 * the analyzer in order to do real-time parameter changes
+	 */
+	class IAnalyzerControls
+	{
+	public:
+		virtual ~IAnalyzerControls() {};
+	};
+
 	// An audio analyzer worker
 	class IAnalyzerWorker
 	{
@@ -53,6 +64,9 @@ namespace Audio
 		virtual ~IAnalyzerWorker() {};
 
 		virtual void Analyze(TArrayView<const float> InAudio, IAnalyzerResult* OutResult) = 0;
+
+		virtual void SetControls(TSharedPtr<IAnalyzerControls> InAnalyzerControls) {}
+		virtual void ClearControls() {}
 	};
 
 	// Audio analyzer factory. Primarily used for creating a worker.
