Disable audio synesthesia NRT analysis for procedural soundwaves including MetaSounds

diff --git a/Engine/Source/Runtime/AudioAnalyzer/Classes/AudioAnalyzerNRT.h b/Engine/Source/Runtime/AudioAnalyzer/Classes/AudioAnalyzerNRT.h
--- a/Engine/Source/Runtime/AudioAnalyzer/Classes/AudioAnalyzerNRT.h
+++ b/Engine/Source/Runtime/AudioAnalyzer/Classes/AudioAnalyzerNRT.h
@@ -77,7 +77,7 @@ class UAudioAnalyzerNRT : public UAudioAnalyzerAssetBase
 		/**
 		 * The USoundWave which is analyzed.
 		 */
-		UPROPERTY(EditAnywhere, BlueprintReadOnly, Category=AudioAnalyzer)
+		UPROPERTY(EditAnywhere, BlueprintReadOnly, Category=AudioAnalyzer, meta=(DisallowedClasses="/Script/MetasoundEngine.MetaSoundSource, /Script/Engine.SoundSourceBus"))
 		TObjectPtr<USoundWave> Sound;
 
 		/** The duration of the analyzed audio in seconds. */
diff --git a/Engine/Source/Runtime/AudioAnalyzer/Private/AudioAnalyzerNRT.cpp b/Engine/Source/Runtime/AudioAnalyzer/Private/AudioAnalyzerNRT.cpp
--- a/Engine/Source/Runtime/AudioAnalyzer/Private/AudioAnalyzerNRT.cpp
+++ b/Engine/Source/Runtime/AudioAnalyzer/Private/AudioAnalyzerNRT.cpp
@@ -125,17 +125,23 @@ void UAudioAnalyzerNRT::AnalyzeAudio()
 {
 	AUDIO_ANALYSIS_LLM_SCOPE
 
-	TSharedPtr<Audio::IAnalyzerNRTResult, ESPMode::ThreadSafe> NewResult;
-
 	// Create a new result id for this result.
 	FResultId ThisResultId = ++CurrentResultId;
 
 	if (nullptr != Sound)
 	{
+		if (Sound->bProcedural)
+		{
+			UE_LOG(LogAudioAnalyzer, Warning, TEXT("Soundwave '%s' is procedural. NRT audio analysis is not currently supported for this."), *Sound->GetFullName());
+			SetResult(nullptr);
+			return;
+		}
+
 		// Read audio while Sound object is assured safe. 
 		if (Sound->ChannelSizes.Num() > 0)
 		{
 			UE_LOG(LogAudioAnalyzer, Warning, TEXT("Soundwave '%s' has multi-channel audio (channels greater than 2). Audio analysis is not currently supported for this yet."), *Sound->GetFullName());
+			SetResult(nullptr);
 			return;
 		}
 
@@ -147,12 +153,14 @@ void UAudioAnalyzerNRT::AnalyzeAudio()
 		if (!Sound->GetImportedSoundWaveData(RawWaveData, SampleRate, NumChannels))
 		{
 			UE_LOG(LogAudioAnalyzer, Error, TEXT("Could not analyze audio due to failed import of sound wave data from Soundwave '%s'."), *Sound->GetFullName());
+			SetResult(nullptr);
 			return;
 		}
 
 		if (SampleRate == 0 || NumChannels == 0)
 		{
 			UE_LOG(LogAudioAnalyzer, Error, TEXT("Failed to parse the raw imported data for '%s' for analysis."), *Sound->GetFullName());
+			SetResult(nullptr);
 			return;
 		}
 		
