Allow sounds to be set on analyzer in editor

diff --git a/Engine/Source/Runtime/AudioAnalyzer/Classes/AudioAnalyzerNRT.h b/Engine/Source/Runtime/AudioAnalyzer/Classes/AudioAnalyzerNRT.h
--- a/Engine/Source/Runtime/AudioAnalyzer/Classes/AudioAnalyzerNRT.h
+++ b/Engine/Source/Runtime/AudioAnalyzer/Classes/AudioAnalyzerNRT.h
@@ -48,6 +48,9 @@ class UAudioAnalyzerNRTSettings : public UAudioAnalyzerAssetBase
 
 };
 
+/** Delegate to receive all meter results per channel (time-stamped in an array) since last delegate call. */
+DECLARE_DYNAMIC_MULTICAST_DELEGATE(FOnAudioAnalyzerNRTAnalysisComplete);
+
 /** UAudioAnalyzerNRT
  *
  * UAudioAnalyzerNRT applies an analyzer to a sound using specific settings, stores the 
@@ -84,6 +87,11 @@ class UAudioAnalyzerNRT : public UAudioAnalyzerAssetBase
 		UPROPERTY(BlueprintReadOnly, Category=AudioAnalyzer)
 		float DurationInSeconds;
 
+#if WITH_EDITOR
+		UFUNCTION(BlueprintCallable, Category="Audio Analyzer")
+		AUDIOANALYZER_API void SetSound(USoundWave* InSound);
+#endif
+
 		/**
 		 * Returns the result object generated by the associated IAudioAnalyzerNRTFactory.
 		 * The template argument must be the IAudioAnalyzerNRTResult subclass returned by
@@ -156,8 +164,14 @@ class UAudioAnalyzerNRT : public UAudioAnalyzerAssetBase
 		AUDIOANALYZER_API virtual bool ShouldEventTriggerAnalysis(struct FPropertyChangedEvent & PropertyChangeEvent);
 
 		/** Performs the analaysis of the audio */
-		UFUNCTION()
+		UFUNCTION(BlueprintCallable, Category="Audio Analyzer")
 		AUDIOANALYZER_API void AnalyzeAudio();
+
+#endif
+#if WITH_EDITORONLY_DATA
+		/** Delegate to receive all meter results, per-channel, since last delegate call. */
+		UPROPERTY(BlueprintAssignable, Category="Audio Analyzer")
+		FOnAudioAnalyzerNRTAnalysisComplete OnAnalysisComplete;
 #endif
 
 
diff --git a/Engine/Source/Runtime/AudioAnalyzer/Private/AudioAnalyzerNRT.cpp b/Engine/Source/Runtime/AudioAnalyzer/Private/AudioAnalyzerNRT.cpp
--- a/Engine/Source/Runtime/AudioAnalyzer/Private/AudioAnalyzerNRT.cpp
+++ b/Engine/Source/Runtime/AudioAnalyzer/Private/AudioAnalyzerNRT.cpp
@@ -81,6 +81,10 @@ bool UAudioAnalyzerNRTSettings::ShouldEventTriggerAnalysis(struct FPropertyChang
 /*****************************************************/
 /***********      UAudioAnalyzerNRT    ***************/
 /*****************************************************/
+void UAudioAnalyzerNRT::SetSound(USoundWave* InSound)
+{
+	Sound = InSound;
+}
 
 void UAudioAnalyzerNRT::PreEditChange(FProperty* PropertyAboutToChange)
 {
@@ -124,6 +128,8 @@ bool UAudioAnalyzerNRT::ShouldEventTriggerAnalysis(struct FPropertyChangedEvent
 void UAudioAnalyzerNRT::AnalyzeAudio()
 {
 	AUDIO_ANALYSIS_LLM_SCOPE
+	
+	check(IsInGameThread());
 
 	// Create a new result id for this result.
 	FResultId ThisResultId = ++CurrentResultId;
@@ -209,6 +215,9 @@ UAudioAnalyzerNRTSettings* UAudioAnalyzerNRT::GetSettingsFromProperty(FProperty*
 
 void UAudioAnalyzerNRT::SetResult(FResultSharedPtr NewResult)
 {
+	check(IsInGameThread());
+
+	
 	FScopeLock ResultLock(&ResultCriticalSection);
 	Result = NewResult;
 	if (Result.IsValid())
@@ -220,10 +229,16 @@ void UAudioAnalyzerNRT::SetResult(FResultSharedPtr NewResult)
 		DurationInSeconds = 0.f;
 	}
 	Modify();
+#if WITH_EDITORONLY_DATA
+	OnAnalysisComplete.Broadcast();
+#endif
+	
+
 }
 
 void UAudioAnalyzerNRT::SetResultIfLatest(FResultSharedPtr NewResult, FResultId InResultId)
 {
+	
 	FScopeLock ResultLock(&ResultCriticalSection);
 	const FResultId ResultId = CurrentResultId.Load();
 
@@ -239,6 +254,9 @@ void UAudioAnalyzerNRT::SetResultIfLatest(FResultSharedPtr NewResult, FResultId
 			DurationInSeconds = 0.f;
 		}
 		Modify();
+#if WITH_EDITORONLY_DATA
+		OnAnalysisComplete.Broadcast();
+#endif
 	}
 }
 
